Fix the "Engine" directory being added as an include path. UBT adds any headers under include paths to the list of build products required for the installed build, and this is causing anything under the Engine/Intermediate folder to be dragged in. Fixes CIS warnings since Dev-VR merge.

diff --git a/Engine/Plugins/Experimental/BackChannel/Source/BackChannel/BackChannel.Build.cs b/Engine/Plugins/Experimental/BackChannel/Source/BackChannel/BackChannel.Build.cs
--- a/Engine/Plugins/Experimental/BackChannel/Source/BackChannel/BackChannel.Build.cs
+++ b/Engine/Plugins/Experimental/BackChannel/Source/BackChannel/BackChannel.Build.cs
@@ -8,7 +8,7 @@ public class BackChannel : ModuleRules
 	{
 		PublicIncludePaths.AddRange(
 			new string[] {
-				".."
+				System.IO.Path.Combine(ModuleDirectory, "..")
 			}
 		);
 
