Disabled BackChannel tests

diff --git a/Engine/Plugins/Experimental/BackChannel/Source/BackChannel/Private/Transport/TestBackChannelTransport.cpp b/Engine/Plugins/Experimental/BackChannel/Source/BackChannel/Private/Transport/TestBackChannelTransport.cpp
--- a/Engine/Plugins/Experimental/BackChannel/Source/BackChannel/Private/Transport/TestBackChannelTransport.cpp
+++ b/Engine/Plugins/Experimental/BackChannel/Source/BackChannel/Private/Transport/TestBackChannelTransport.cpp
@@ -8,7 +8,7 @@
 
 PRAGMA_DISABLE_OPTIMIZATION
 
-#if WITH_DEV_AUTOMATION_TESTS && 1
+#if WITH_DEV_AUTOMATION_TESTS && 0
 
 class FBackChannelTestTransport : public FAutomationTestBase
 {
