UE-61289 RemoteSession - Unable to reconnect to a PIE session on Mac with an Android device after killing a PIE instance)

#Android
#4.20

Restarting the host app on MAC, not 100%:  the newly created TCP socket raises an error at Bind(), complaining that the address is still in use.
Using NewSocket->SetReuseAddr(true) solves the problem.

diff --git a/Engine/Plugins/Experimental/BackChannel/Source/BackChannel/Private/Transport/BackChannelConnection.cpp b/Engine/Plugins/Experimental/BackChannel/Source/BackChannel/Private/Transport/BackChannelConnection.cpp
--- a/Engine/Plugins/Experimental/BackChannel/Source/BackChannel/Private/Transport/BackChannelConnection.cpp
+++ b/Engine/Plugins/Experimental/BackChannel/Source/BackChannel/Private/Transport/BackChannelConnection.cpp
@@ -180,7 +180,7 @@ bool FBackChannelConnection::Listen(const int16 Port)
         
         if (NewSocket != nullptr)
         {
-            bool Error = false;//!NewSocket->SetReuseAddr(true);
+            bool Error = !NewSocket->SetReuseAddr(true);
             
             if (!Error)
             {
