UE-61351 Error: Failed to send -1 bytes of data to RemoteConnection spammed in the editor log when closing RemoteSession on device during connection

#4.21

Not an real issue. The message is issued whenever the host (or the client) tries to send messages through an inactive channel. There is a "/ping" message sent back and forth between the host and the client, with high frequency (hence the spam) - if this message is no longer received, the host (or the client) decides to terminate the connection.

"Downgrading" the message from Error to Log, to avoid spamming the log on release builds.

diff --git a/Engine/Plugins/Experimental/BackChannel/Source/BackChannel/Private/Transport/BackChannelConnection.cpp b/Engine/Plugins/Experimental/BackChannel/Source/BackChannel/Private/Transport/BackChannelConnection.cpp
--- a/Engine/Plugins/Experimental/BackChannel/Source/BackChannel/Private/Transport/BackChannelConnection.cpp
+++ b/Engine/Plugins/Experimental/BackChannel/Source/BackChannel/Private/Transport/BackChannelConnection.cpp
@@ -355,7 +355,7 @@ int32 FBackChannelConnection::SendData(const void* InData, const int32 InSize)
 			//const TCHAR* SocketErr = ISocketSubsystem::Get(PLATFORM_SOCKETSUBSYSTEM)->GetSocketError(LastError);
 			//UE_CLOG(GBackChannelLogErrors, LogBackChannel, Error, TEXT("Failed to send %d bytes of data to %s. Err: %s"), BytesSent, *GetDescription(), SocketErr);
             
-            UE_CLOG(GBackChannelLogErrors, LogBackChannel, Error, TEXT("Failed to send %d bytes of data to %s. "), BytesSent, *GetDescription());
+            UE_CLOG(GBackChannelLogErrors, LogBackChannel, Log, TEXT("Failed to send %d bytes of data to %s. "), BytesSent, *GetDescription());
 		}
 	}
 	else
