Removed the Tcp builder dependency from back channel since that hasn't been updated properly for inet family changes

diff --git a/Engine/Plugins/Experimental/BackChannel/Source/BackChannel/Private/Transport/BackChannelConnection.cpp b/Engine/Plugins/Experimental/BackChannel/Source/BackChannel/Private/Transport/BackChannelConnection.cpp
--- a/Engine/Plugins/Experimental/BackChannel/Source/BackChannel/Private/Transport/BackChannelConnection.cpp
+++ b/Engine/Plugins/Experimental/BackChannel/Source/BackChannel/Private/Transport/BackChannelConnection.cpp
@@ -12,7 +12,9 @@
 DECLARE_DWORD_COUNTER_STAT(TEXT("BCBytesSent"), STAT_BackChannelBytesSent, STATGROUP_Game);
 DECLARE_DWORD_COUNTER_STAT(TEXT("BCBytesRecv"), STAT_BackChannelBytesRecv, STATGROUP_Game);
 
-static const int32 DefaultBufferSize = 2 * 1024 * 1024;
+
+int32 FBackChannelConnection::SendBufferSize = 2 * 1024 * 1024;
+int32 FBackChannelConnection::ReceiveBufferSize = 2 * 1024 * 1024;
 
 int32 GBackChannelLogPackets = 0;
 static FAutoConsoleVariableRef BCCVarLogPackets(
@@ -31,6 +33,9 @@ FBackChannelConnection::FBackChannelConnection()
 	Socket = nullptr;
 	IsListener = false;
 	PacketsReceived = 0;
+	// Allow the app to override
+	GConfig->GetInt(TEXT("BackChannel"), TEXT("SendBufferSize"), SendBufferSize, GEngineIni);
+	GConfig->GetInt(TEXT("BackChannel"), TEXT("RecvBufferSize"), ReceiveBufferSize, GEngineIni);
 }
 
 FBackChannelConnection::~FBackChannelConnection()
@@ -109,33 +114,44 @@ bool FBackChannelConnection::Connect(const TCHAR* InEndPoint)
 
 	FString LocalEndPoint = InEndPoint;
 
-	FString Description = FString::Printf(TEXT("%s"), InEndPoint);
+	FSocket* NewSocket = ISocketSubsystem::Get(PLATFORM_SOCKETSUBSYSTEM)->CreateSocket(NAME_Stream, TEXT("FBackChannelConnection Client Socket"));
+	if (NewSocket)
+	{
+		NewSocket->SetNonBlocking();
 
-	int32 ReceiveBufferSize = DefaultBufferSize;
-	int32 SendBufferSize = DefaultBufferSize;
-	// Allow the app to override
-	GConfig->GetInt(TEXT("BackChannel"), TEXT("SendBufferSize"), SendBufferSize, GEngineIni);
-	GConfig->GetInt(TEXT("BackChannel"), TEXT("RecvBufferSize"), ReceiveBufferSize, GEngineIni);
+		int32 NewSize = 0;
+		NewSocket->SetSendBufferSize(SendBufferSize, NewSize);
+		if (NewSize != SendBufferSize)
+		{
+			UE_LOG(LogBackChannel, Log, TEXT("SetSendBufferSize requested (%d) size but got (%d) size"), SendBufferSize, NewSize);
+		}
+		NewSocket->SetReceiveBufferSize(ReceiveBufferSize, NewSize);
+		if (NewSize != ReceiveBufferSize)
+		{
+			UE_LOG(LogBackChannel, Log, TEXT("SetReceiveBufferSize requested (%d) size but got (%d) size"), SendBufferSize, NewSize);
+		}
 
-	FSocket* NewSocket = FTcpSocketBuilder(*Description)
-		.WithSendBufferSize(SendBufferSize)
-		.WithReceiveBufferSize(ReceiveBufferSize);
+		bool Success = false;
+
+		FIPv4Endpoint EndPointv4;
+		// Check for a valid IPv4 address string
+		if (FIPv4Endpoint::Parse(LocalEndPoint, EndPointv4))
+		{
+			Success = NewSocket->Connect(*EndPointv4.ToInternetAddr());
+		}
 
-	if (NewSocket)
-	{
-		FIPv4Endpoint EndPoint;
-		FIPv4Endpoint::Parse(LocalEndPoint, EndPoint);
-        
-        bool Success = NewSocket->Connect(*EndPoint.ToInternetAddr());
-        
         if (!Success)
         {
-            int32 LastErr = ISocketSubsystem::Get(PLATFORM_SOCKETSUBSYSTEM)->GetLastErrorCode();
+            ESocketErrors LastErr = ISocketSubsystem::Get(PLATFORM_SOCKETSUBSYSTEM)->GetLastErrorCode();
             
             if (LastErr == SE_EINPROGRESS || LastErr == SE_EWOULDBLOCK)
             {
                 Success = true;
             }
+			else
+			{
+				UE_LOG(LogBackChannel, Log, TEXT("Connect failed with error code (%d) error (%s)"), LastErr, ISocketSubsystem::Get(PLATFORM_SOCKETSUBSYSTEM)->GetSocketError(LastErr));
+			}
         }
 
 		if (Success)
@@ -157,77 +173,59 @@ bool FBackChannelConnection::Listen(const int16 Port)
 {
 	FScopeLock Lock(&SocketMutex);
 
-	FIPv4Endpoint Endpoint(FIPv4Address::Any, Port);
-    
-    FSocket* NewSocket = nullptr;
+	ISocketSubsystem* SocketSubsystem = ISocketSubsystem::Get(PLATFORM_SOCKETSUBSYSTEM);
+	FSocket* NewSocket = SocketSubsystem->CreateSocket(NAME_Stream, TEXT("FBackChannelConnection Client Socket"));
+	if (NewSocket != nullptr)
+	{
+		FIPv4Endpoint Endpoint(FIPv4Address::Any, Port);
 
-	int32 ReceiveBufferSize = DefaultBufferSize;
-	int32 SendBufferSize = DefaultBufferSize;
-	// Allow the app to override
-	GConfig->GetInt(TEXT("BackChannel"), TEXT("SendBufferSize"), SendBufferSize, GEngineIni);
-	GConfig->GetInt(TEXT("BackChannel"), TEXT("RecvBufferSize"), ReceiveBufferSize, GEngineIni);
+		bool Error = !NewSocket->SetReuseAddr(true);
 
-	NewSocket = FTcpSocketBuilder(TEXT("FBackChannelConnection ListenSocket"))
-		.BoundToEndpoint(Endpoint)
-		.Listening(8)
-		.WithSendBufferSize(DefaultBufferSize);
-    
-    ISocketSubsystem* SocketSubsystem = ISocketSubsystem::Get(PLATFORM_SOCKETSUBSYSTEM);
-    
-    if (NewSocket == nullptr && SocketSubsystem != nullptr)
-    {
-        NewSocket = SocketSubsystem->CreateSocket(NAME_Stream, TEXT("FBackChannelConnection ListenSocket"), true);
-        
-        if (NewSocket != nullptr)
-        {
-            bool Error = !NewSocket->SetReuseAddr(true);
-            
-            if (!Error)
-            {
-                Error = !NewSocket->SetRecvErr();
-            }
-            
-            if (!Error)
-            {
-                Error = !NewSocket->Bind(*Endpoint.ToInternetAddr());
-            }
-            
-            if (!Error)
-            {
-                Error = !NewSocket->Listen(2);
-            }
-            
-            if (!Error)
-            {
-                Error = !NewSocket->SetNonBlocking(!false);
-            }
-            
-            if (!Error)
-            {
-                int32 OutNewSize;
-
-                if (ReceiveBufferSize > 0) //-V547
-                {
-                    NewSocket->SetReceiveBufferSize(ReceiveBufferSize, OutNewSize);
-                }
-                
-                if (SendBufferSize > 0) //-V547
-                {
-                    NewSocket->SetSendBufferSize(SendBufferSize, OutNewSize);
-                }
-            }
-            
-            if (Error)
-            {
-                const TCHAR* SocketErr = ISocketSubsystem::Get(PLATFORM_SOCKETSUBSYSTEM)->GetSocketError(SE_GET_LAST_ERROR_CODE);
-                GLog->Logf(TEXT("FTcpSocketBuilder: Failed to create the socket %s as configured. %s"), TEXT("Pfft"), SocketErr);
-                
-                SocketSubsystem->DestroySocket(NewSocket);
-                
-                NewSocket = nullptr;
-            }
-        }
-    }
+		if (!Error)
+		{
+			Error = !NewSocket->SetRecvErr();
+		}
+
+		if (!Error)
+		{
+			Error = !NewSocket->SetNonBlocking();
+		}
+
+		if (!Error)
+		{
+			int32 NewSize = 0;
+			NewSocket->SetSendBufferSize(SendBufferSize, NewSize);
+			if (NewSize != SendBufferSize)
+			{
+				UE_LOG(LogBackChannel, Log, TEXT("SetSendBufferSize requested (%d) size but got (%d) size"), SendBufferSize, NewSize);
+			}
+			NewSocket->SetReceiveBufferSize(ReceiveBufferSize, NewSize);
+			if (NewSize != ReceiveBufferSize)
+			{
+				UE_LOG(LogBackChannel, Log, TEXT("SetReceiveBufferSize requested (%d) size but got (%d) size"), SendBufferSize, NewSize);
+			}
+		}
+
+		if (!Error)
+		{
+			Error = !NewSocket->Bind(*Endpoint.ToInternetAddr());
+		}
+
+		if (!Error)
+		{
+			Error = !NewSocket->Listen(8);
+		}
+
+		if (Error)
+		{
+			const TCHAR* SocketErr = ISocketSubsystem::Get(PLATFORM_SOCKETSUBSYSTEM)->GetSocketError(SE_GET_LAST_ERROR_CODE);
+			GLog->Logf(TEXT("Failed to create the listen socket as configured. %s"), SocketErr);
+
+			SocketSubsystem->DestroySocket(NewSocket);
+
+			NewSocket = nullptr;
+		}
+	}
 
 	if (NewSocket == nullptr)
 	{
@@ -301,6 +299,23 @@ bool FBackChannelConnection::WaitForConnection(double InTimeout, TFunction<bool(
 
 				if (ConnectionSocket != nullptr)
 				{
+					// Each platform can inherit different socket options from the listen socket so set ours again
+					{
+						ConnectionSocket->SetNonBlocking();
+
+						int32 NewSize = 0;
+						ConnectionSocket->SetSendBufferSize(SendBufferSize, NewSize);
+						if (NewSize != SendBufferSize)
+						{
+							UE_LOG(LogBackChannel, Log, TEXT("SetSendBufferSize requested (%d) size but got (%d) size"), SendBufferSize, NewSize);
+						}
+						ConnectionSocket->SetReceiveBufferSize(ReceiveBufferSize, NewSize);
+						if (NewSize != ReceiveBufferSize)
+						{
+							UE_LOG(LogBackChannel, Log,TEXT("SetReceiveBufferSize requested (%d) size but got (%d) size"), SendBufferSize,NewSize);
+						}
+					}
+
 					TSharedRef<FBackChannelConnection> BCConnection = MakeShareable(new FBackChannelConnection);
 					BCConnection->Attach(ConnectionSocket);
 
@@ -351,11 +366,9 @@ int32 FBackChannelConnection::SendData(const void* InData, const int32 InSize)
 	{
 		if (GBackChannelLogErrors)
 		{
-            //ESocketErrors LastError = ISocketSubsystem::Get(PLATFORM_SOCKETSUBSYSTEM)->GetLastErrorCode();
-			//const TCHAR* SocketErr = ISocketSubsystem::Get(PLATFORM_SOCKETSUBSYSTEM)->GetSocketError(LastError);
-			//UE_CLOG(GBackChannelLogErrors, LogBackChannel, Error, TEXT("Failed to send %d bytes of data to %s. Err: %s"), BytesSent, *GetDescription(), SocketErr);
-            
-            UE_CLOG(GBackChannelLogErrors, LogBackChannel, Log, TEXT("Failed to send %d bytes of data to %s. "), BytesSent, *GetDescription());
+            ESocketErrors LastError = ISocketSubsystem::Get(PLATFORM_SOCKETSUBSYSTEM)->GetLastErrorCode();
+			const TCHAR* SocketErr = ISocketSubsystem::Get(PLATFORM_SOCKETSUBSYSTEM)->GetSocketError(LastError);
+			UE_CLOG(GBackChannelLogErrors, LogBackChannel, Error, TEXT("Failed to send %d bytes of data to %s. Err: %s"), InSize, *GetDescription(), SocketErr);
 		}
 	}
 	else
diff --git a/Engine/Plugins/Experimental/BackChannel/Source/BackChannel/Private/Transport/BackChannelConnection.h b/Engine/Plugins/Experimental/BackChannel/Source/BackChannel/Private/Transport/BackChannelConnection.h
--- a/Engine/Plugins/Experimental/BackChannel/Source/BackChannel/Private/Transport/BackChannelConnection.h
+++ b/Engine/Plugins/Experimental/BackChannel/Source/BackChannel/Private/Transport/BackChannelConnection.h
@@ -55,6 +55,8 @@ public:
 	uint32	GetPacketsReceived() const override;
 
 private:
+	static int32 SendBufferSize;
+	static int32 ReceiveBufferSize;
 
 	void					CloseWithError(const TCHAR* Error, FSocket* InSocket=nullptr);
 
