Fixed issue with RemoteSessionApp overflowing send buffers and disconnecting when Mac is a host.

At somepoint the default send/receive buffer for RemoteSessionApp was increased to 8mb (presumably to support higher quality/framerate) but MacOS doesn't allow that size buffer so the buffers were left at the OS default of 128KB

Now if the requested size fails we keep trying at 50% until one succeeds. This gives 4MB on MacOS which is perflectly stable at 60fps / iPhone XS size.

diff --git a/Engine/Plugins/Experimental/BackChannel/Source/BackChannel/Private/Transport/BackChannelConnection.cpp b/Engine/Plugins/Experimental/BackChannel/Source/BackChannel/Private/Transport/BackChannelConnection.cpp
--- a/Engine/Plugins/Experimental/BackChannel/Source/BackChannel/Private/Transport/BackChannelConnection.cpp
+++ b/Engine/Plugins/Experimental/BackChannel/Source/BackChannel/Private/Transport/BackChannelConnection.cpp
@@ -9,8 +9,10 @@
 #include "CoreGlobals.h"
 #include "Misc/ConfigCacheIni.h"
 
-DECLARE_DWORD_COUNTER_STAT(TEXT("BCBytesSent"), STAT_BackChannelBytesSent, STATGROUP_Game);
-DECLARE_DWORD_COUNTER_STAT(TEXT("BCBytesRecv"), STAT_BackChannelBytesRecv, STATGROUP_Game);
+DECLARE_STATS_GROUP(TEXT("BackChannel"), STATGROUP_BackChannel, STATCAT_Advanced);
+
+DECLARE_DWORD_COUNTER_STAT(TEXT("BC_BytesSent"), STAT_BackChannelBytesSent, STATGROUP_BackChannel);
+DECLARE_DWORD_COUNTER_STAT(TEXT("BC_BytesRecv"), STAT_BackChannelBytesRecv, STATGROUP_BackChannel);
 
 int32 FBackChannelConnection::SendBufferSize = 2 * 1024 * 1024;
 int32 FBackChannelConnection::ReceiveBufferSize = 2 * 1024 * 1024;
@@ -168,6 +170,47 @@ bool FBackChannelConnection::Connect(const TCHAR* InEndPoint)
 	return Socket != nullptr;
 }
 
+void FBackChannelConnection::SetSocketBufferSizes(FSocket* NewSocket, int32 DesiredSendSize, int32 DesiredReceiveSize)
+{
+	int32 AllocatedSendSize = 0;
+	int32 AllocatedReceiveSize = 0;
+
+	int32 RequestedSendSize = DesiredSendSize;
+	int32 RequestedReceiveSize = DesiredReceiveSize;
+	
+	// Send Buffer
+	while (AllocatedSendSize != RequestedSendSize)
+	{
+		NewSocket->SetSendBufferSize(RequestedSendSize, AllocatedSendSize);
+		
+		if (AllocatedSendSize != RequestedSendSize)
+		{
+			RequestedSendSize = RequestedSendSize / 2;
+		}
+	}
+	
+	if (AllocatedSendSize != DesiredSendSize)
+	{
+		UE_LOG(LogBackChannel, Warning, TEXT("Wanted send buffer of %d but OS only allowed %d"), DesiredSendSize, AllocatedSendSize);
+	}
+	
+	// Set Receive buffer
+	while (AllocatedReceiveSize != RequestedReceiveSize)
+	{
+		NewSocket->SetReceiveBufferSize(RequestedReceiveSize, AllocatedReceiveSize);
+		
+		if (AllocatedReceiveSize != RequestedReceiveSize)
+		{
+			RequestedReceiveSize = RequestedReceiveSize / 2;
+		}
+	}
+	
+	if (AllocatedReceiveSize != DesiredReceiveSize)
+	{
+		UE_LOG(LogBackChannel, Warning, TEXT("Wanted receive buffer of %d but OS only allowed %d"), DesiredReceiveSize, AllocatedReceiveSize);
+	}
+}
+
 bool FBackChannelConnection::Listen(const int16 Port)
 {
 	FScopeLock Lock(&SocketMutex);
@@ -192,17 +235,8 @@ bool FBackChannelConnection::Listen(const int16 Port)
 
 		if (!Error)
 		{
-			int32 NewSize = 0;
-			NewSocket->SetSendBufferSize(SendBufferSize, NewSize);
-			if (NewSize != SendBufferSize)
-			{
-				UE_LOG(LogBackChannel, Log, TEXT("SetSendBufferSize requested (%d) size but got (%d) size"), SendBufferSize, NewSize);
-			}
-			NewSocket->SetReceiveBufferSize(ReceiveBufferSize, NewSize);
-			if (NewSize != ReceiveBufferSize)
-			{
-				UE_LOG(LogBackChannel, Log, TEXT("SetReceiveBufferSize requested (%d) size but got (%d) size"), SendBufferSize, NewSize);
-			}
+			// set buffer sizes
+			SetSocketBufferSizes(NewSocket, SendBufferSize, ReceiveBufferSize);
 		}
 
 		if (!Error)
diff --git a/Engine/Plugins/Experimental/BackChannel/Source/BackChannel/Private/Transport/BackChannelConnection.h b/Engine/Plugins/Experimental/BackChannel/Source/BackChannel/Private/Transport/BackChannelConnection.h
--- a/Engine/Plugins/Experimental/BackChannel/Source/BackChannel/Private/Transport/BackChannelConnection.h
+++ b/Engine/Plugins/Experimental/BackChannel/Source/BackChannel/Private/Transport/BackChannelConnection.h
@@ -59,6 +59,9 @@ private:
 	static int32 ReceiveBufferSize;
 
 	void					CloseWithError(const TCHAR* Error, FSocket* InSocket=nullptr);
+	
+	/* Attempts to set the specified buffer size on our socket, will drop by 50% each time until success */
+	void					SetSocketBufferSizes(FSocket* NewSocket, int32 DesiredSendSize, int32 DesiredReceiveSize);
 
 	FThreadSafeBool			IsAttemptingConnection;
 	FCriticalSection		SocketMutex;
