On Linux when setsockopt is used with a larger values then the max it will set the max value and return it was set (true)
On Mac when setsockopt is used with a larger value then the max it will not set the max value and will return it was not set (false)

To avoid an infinite loop, if we did set a value assume that value was as correct as we could get it to the requested size

Warnings will still print as the size set will not equal the requested size in if the requested size is to large

diff --git a/Engine/Plugins/Experimental/BackChannel/Source/BackChannel/Private/Transport/BackChannelConnection.cpp b/Engine/Plugins/Experimental/BackChannel/Source/BackChannel/Private/Transport/BackChannelConnection.cpp
--- a/Engine/Plugins/Experimental/BackChannel/Source/BackChannel/Private/Transport/BackChannelConnection.cpp
+++ b/Engine/Plugins/Experimental/BackChannel/Source/BackChannel/Private/Transport/BackChannelConnection.cpp
@@ -177,12 +177,13 @@ void FBackChannelConnection::SetSocketBufferSizes(FSocket* NewSocket, int32 Desi
 
 	int32 RequestedSendSize = DesiredSendSize;
 	int32 RequestedReceiveSize = DesiredReceiveSize;
+	bool bWasSet = false;
 	
 	// Send Buffer
-	while (AllocatedSendSize != RequestedSendSize)
+	while (AllocatedSendSize != RequestedSendSize && !bWasSet)
 	{
-		NewSocket->SetSendBufferSize(RequestedSendSize, AllocatedSendSize);
-		
+		bWasSet = NewSocket->SetSendBufferSize(RequestedSendSize, AllocatedSendSize);
+
 		if (AllocatedSendSize != RequestedSendSize)
 		{
 			RequestedSendSize = RequestedSendSize / 2;
@@ -194,10 +195,12 @@ void FBackChannelConnection::SetSocketBufferSizes(FSocket* NewSocket, int32 Desi
 		UE_LOG(LogBackChannel, Warning, TEXT("Wanted send buffer of %d but OS only allowed %d"), DesiredSendSize, AllocatedSendSize);
 	}
 	
+	bWasSet = false;
+
 	// Set Receive buffer
-	while (AllocatedReceiveSize != RequestedReceiveSize)
+	while (AllocatedReceiveSize != RequestedReceiveSize && !bWasSet)
 	{
-		NewSocket->SetReceiveBufferSize(RequestedReceiveSize, AllocatedReceiveSize);
+		bWasSet = NewSocket->SetReceiveBufferSize(RequestedReceiveSize, AllocatedReceiveSize);
 		
 		if (AllocatedReceiveSize != RequestedReceiveSize)
 		{
