Fixed addirional backchannel connection that was not having socket size set robustly.

diff --git a/Engine/Plugins/Experimental/BackChannel/Source/BackChannel/Private/Transport/BackChannelConnection.cpp b/Engine/Plugins/Experimental/BackChannel/Source/BackChannel/Private/Transport/BackChannelConnection.cpp
--- a/Engine/Plugins/Experimental/BackChannel/Source/BackChannel/Private/Transport/BackChannelConnection.cpp
+++ b/Engine/Plugins/Experimental/BackChannel/Source/BackChannel/Private/Transport/BackChannelConnection.cpp
@@ -120,17 +120,8 @@ bool FBackChannelConnection::Connect(const TCHAR* InEndPoint)
 	{
 		NewSocket->SetNonBlocking();
 
-		int32 NewSize = 0;
-		NewSocket->SetSendBufferSize(SendBufferSize, NewSize);
-		if (NewSize != SendBufferSize)
-		{
-			UE_LOG(LogBackChannel, Log, TEXT("SetSendBufferSize requested (%d) size but got (%d) size"), SendBufferSize, NewSize);
-		}
-		NewSocket->SetReceiveBufferSize(ReceiveBufferSize, NewSize);
-		if (NewSize != ReceiveBufferSize)
-		{
-			UE_LOG(LogBackChannel, Log, TEXT("SetReceiveBufferSize requested (%d) size but got (%d) size"), SendBufferSize, NewSize);
-		}
+		// set buffer sizes
+		SetSocketBufferSizes(NewSocket, SendBufferSize, ReceiveBufferSize);
 
 		bool Success = false;
 
@@ -178,13 +169,16 @@ void FBackChannelConnection::SetSocketBufferSizes(FSocket* NewSocket, int32 Desi
 	int32 RequestedSendSize = DesiredSendSize;
 	int32 RequestedReceiveSize = DesiredReceiveSize;
 	bool bWasSet = false;
-	
+		
 	// Send Buffer
 	while (AllocatedSendSize != RequestedSendSize && !bWasSet)
 	{
 		bWasSet = NewSocket->SetSendBufferSize(RequestedSendSize, AllocatedSendSize);
 
-		if (AllocatedSendSize != RequestedSendSize)
+		// If we didn't get what we want assume failure could mean
+		// no change (unsupported size), an OS default, or an OS max so
+		// try again if its less than 50% of what we asked for
+		if (!bWasSet && AllocatedSendSize < RequestedSendSize/2)
 		{
 			RequestedSendSize = RequestedSendSize / 2;
 		}
@@ -192,7 +186,11 @@ void FBackChannelConnection::SetSocketBufferSizes(FSocket* NewSocket, int32 Desi
 	
 	if (AllocatedSendSize != DesiredSendSize)
 	{
-		UE_LOG(LogBackChannel, Warning, TEXT("Wanted send buffer of %d but OS only allowed %d"), DesiredSendSize, AllocatedSendSize);
+		UE_LOG(LogBackChannel, Warning, TEXT("Wanted send buffer of %d for %s but OS only allowed %d"), DesiredSendSize, *NewSocket->GetDescription(), AllocatedSendSize);
+	}
+	else
+	{
+		UE_LOG(LogBackChannel, Log, TEXT("Set send buffer to %d bytes for %s"), AllocatedSendSize, *NewSocket->GetDescription());
 	}
 	
 	bWasSet = false;
@@ -202,7 +200,10 @@ void FBackChannelConnection::SetSocketBufferSizes(FSocket* NewSocket, int32 Desi
 	{
 		bWasSet = NewSocket->SetReceiveBufferSize(RequestedReceiveSize, AllocatedReceiveSize);
 		
-		if (AllocatedReceiveSize != RequestedReceiveSize)
+		// If we didn't get what we want assume failure could mean
+		// no change (unsupported size), an OS default, or an OS max so
+		// try again if its less than 50% of what we asked for
+		if (!bWasSet && AllocatedReceiveSize < (RequestedReceiveSize/2))
 		{
 			RequestedReceiveSize = RequestedReceiveSize / 2;
 		}
@@ -210,7 +211,11 @@ void FBackChannelConnection::SetSocketBufferSizes(FSocket* NewSocket, int32 Desi
 	
 	if (AllocatedReceiveSize != DesiredReceiveSize)
 	{
-		UE_LOG(LogBackChannel, Warning, TEXT("Wanted receive buffer of %d but OS only allowed %d"), DesiredReceiveSize, AllocatedReceiveSize);
+		UE_LOG(LogBackChannel, Warning, TEXT("Wanted receive buffer of %d for %s but OS only allowed %d"), DesiredReceiveSize, *NewSocket->GetDescription(), AllocatedReceiveSize);
+	}
+	else
+	{
+		UE_LOG(LogBackChannel, Log, TEXT("Set receive buffer to %d bytes for %s"), AllocatedSendSize, *NewSocket->GetDescription());
 	}
 }
 
@@ -236,11 +241,6 @@ bool FBackChannelConnection::Listen(const int16 Port)
 			Error = !NewSocket->SetNonBlocking();
 		}
 
-		if (!Error)
-		{
-			// set buffer sizes
-			SetSocketBufferSizes(NewSocket, SendBufferSize, ReceiveBufferSize);
-		}
 
 		if (!Error)
 		{
@@ -339,17 +339,8 @@ bool FBackChannelConnection::WaitForConnection(double InTimeout, TFunction<bool(
 					{
 						ConnectionSocket->SetNonBlocking();
 
-						int32 NewSize = 0;
-						ConnectionSocket->SetSendBufferSize(SendBufferSize, NewSize);
-						if (NewSize != SendBufferSize)
-						{
-							UE_LOG(LogBackChannel, Log, TEXT("SetSendBufferSize requested (%d) size but got (%d) size"), SendBufferSize, NewSize);
-						}
-						ConnectionSocket->SetReceiveBufferSize(ReceiveBufferSize, NewSize);
-						if (NewSize != ReceiveBufferSize)
-						{
-							UE_LOG(LogBackChannel, Log, TEXT("SetReceiveBufferSize requested (%d) size but got (%d) size"), SendBufferSize, NewSize);
-						}
+						// set buffer sizes
+						SetSocketBufferSizes(ConnectionSocket, SendBufferSize, ReceiveBufferSize);
 					}
 
 					TSharedRef<FBackChannelConnection> BCConnection = MakeShareable(new FBackChannelConnection);
