BackChannel Plugin: Move files around so the module parent directory doesn't need to be added to public includes. Files are otherwise unmodified

#rnx

diff --git a/Engine/Plugins/Experimental/BackChannel/Source/BackChannel/BackChannel.Build.cs b/Engine/Plugins/Experimental/BackChannel/Source/BackChannel/BackChannel.Build.cs
--- a/Engine/Plugins/Experimental/BackChannel/Source/BackChannel/BackChannel.Build.cs
+++ b/Engine/Plugins/Experimental/BackChannel/Source/BackChannel/BackChannel.Build.cs
@@ -6,13 +6,6 @@ public class BackChannel : ModuleRules
 {
 	public BackChannel( ReadOnlyTargetRules Target ) : base( Target )
 	{
-		PublicIncludePaths.AddRange(
-			new string[] {
-				System.IO.Path.Combine(ModuleDirectory, ".."),
-				System.IO.Path.Combine(ModuleDirectory, "Public")
-			}
-		);
-
 		PrivateDependencyModuleNames.AddRange
 			(
 			new string[] {
diff --git a/Engine/Plugins/Experimental/BackChannel/Source/BackChannel/Private/Utils/BackChannelThreadedConnection.cpp b/Engine/Plugins/Experimental/BackChannel/Source/BackChannel/Private/Utils/BackChannelThreadedConnection.cpp
new file mode 100644
--- /dev/null
+++ b/Engine/Plugins/Experimental/BackChannel/Source/BackChannel/Private/Utils/BackChannelThreadedConnection.cpp
@@ -0,0 +1,60 @@
+// Copyright Epic Games, Inc. All Rights Reserved.
+
+#include "BackChannel/Utils/BackChannelThreadedConnection.h"
+#include "BackChannelCommon.h"
+#include "HAL/RunnableThread.h"
+#include "HAL/ThreadSafeBool.h"
+#include "Misc/ScopeLock.h"
+
+
+FBackChannelThreadedListener::FBackChannelThreadedListener()
+{
+}
+
+FBackChannelThreadedListener::~FBackChannelThreadedListener()
+{
+	Stop();
+}
+
+bool FBackChannelThreadedListener::IsRunning() const
+{
+	return bIsRunning;
+}
+
+void FBackChannelThreadedListener::Start(TSharedRef<IBackChannelSocketConnection> InConnection, FBackChannelListenerDelegate InDelegate)
+{
+	Connection = InConnection;
+	Delegate = InDelegate;
+
+	bIsRunning = true;
+	bExitRequested = false;
+
+	FRunnableThread::Create(this, TEXT("FBackChannelSocketThread"), 32 * 1024);
+}
+
+uint32 FBackChannelThreadedListener::Run()
+{
+	while (bExitRequested == false)
+	{
+		FScopeLock RunningLock(&RunningCS);
+
+		Connection->WaitForConnection(1, [this](TSharedRef<IBackChannelSocketConnection> NewConnection) {
+			return Delegate.Execute(NewConnection);
+		});
+	}
+
+	bIsRunning = false;
+	return 0;
+}
+
+void FBackChannelThreadedListener::Stop()
+{
+	bExitRequested = true;
+
+	if (IsRunning())
+	{
+		FScopeLock RunLock(&RunningCS);
+	}
+
+	bExitRequested = false;
+}
diff --git a/Engine/Plugins/Experimental/BackChannel/Source/BackChannel/Public/BackChannel/Transport/IBackChannelSocketConnection.h b/Engine/Plugins/Experimental/BackChannel/Source/BackChannel/Public/BackChannel/Transport/IBackChannelSocketConnection.h
new file mode 100644
--- /dev/null
+++ b/Engine/Plugins/Experimental/BackChannel/Source/BackChannel/Public/BackChannel/Transport/IBackChannelSocketConnection.h
@@ -0,0 +1,84 @@
+// Copyright Epic Games, Inc. All Rights Reserved.
+
+#pragma once
+
+#include "CoreMinimal.h"
+#include "HAL/PlatformTime.h"
+
+class FSocket;
+
+// todo (agrant 2017/12/29): concept of 'connection' should be a base class with persistent connection subclass?
+
+/**
+ *	Base class that describes a back-channel connection. The underlying behavior will depend on the type
+ *	of connection that was requested from the factory
+ */
+class IBackChannelSocketConnection
+{
+public:
+
+	struct FConnectionStats
+	{
+		// lifetime stats
+		int		PacketsSent = 0;
+		int		PacketsReceived = 0;
+		int		BytesSent = 0;
+		int		BytesReceived = 0;
+		int		Errors = 0;
+
+		// time these events occurred
+		double		LastSendTime = 0;
+		double		LastReceiveTime = 0;
+		double		LastErrorTime = FPlatformTime::Seconds();
+		double		LastSuccessTime = FPlatformTime::Seconds();
+	};
+
+	
+	// todo (agrant 2017/12/29): Should remove 'Connect' and instead return a connected (or null..) socket
+	// from the factory
+
+	/* Start connecting to the specified port for incoming connections. Use WaitForConnection to check status. */
+	virtual bool Connect(const TCHAR* InEndPoint) = 0;
+
+	/* Start listening on the specified port for incoming connections. Use WaitForConnection to accept one. */
+	virtual bool Listen(const int16 Port) = 0;
+
+	/* Close the connection */
+	virtual void Close() = 0;
+
+	/* Waits for an icoming or outgoing connection to be made */
+	virtual bool WaitForConnection(double InTimeout, TFunction<bool(TSharedRef<IBackChannelSocketConnection>)> InDelegate) = 0;
+
+	/* Returns true if this connection is currently listening for incoming connections */
+	virtual bool IsListening() const = 0;
+
+	/* Returns true if this connection is connected to another */
+	virtual bool IsConnected() const = 0;
+
+	/* Send data via our connection */
+	virtual int32 SendData(const void* InData, const int32 InSize) = 0;
+	
+	/* Receive data from our connection. The returned value is the number of bytes read, to a max of BufferSize */
+	virtual int32 ReceiveData(void* OutBuffer, const int32 BufferSize) = 0;
+
+	/* Return the underlying socket (if any) for this connection */
+	virtual FString GetDescription() const = 0;
+
+	/* Return the underlying socket (if any) for this connection */
+	virtual FSocket* GetSocket() = 0;
+
+	/* Todo - Proper stats */
+	virtual uint32	GetPacketsReceived() const = 0;
+
+	/* Set the specified send and receive buffer sizes, if supported */
+	virtual void SetBufferSizes(int32 DesiredSendSize, int32 DesiredReceiveSize) = 0;
+
+	/* Return stats about this connection */
+	virtual const FConnectionStats& GetConnectionStats() const = 0;
+
+protected:
+
+	IBackChannelSocketConnection() {}
+	virtual ~IBackChannelSocketConnection() {}
+};
+
diff --git a/Engine/Plugins/Experimental/BackChannel/Source/BackChannel/Public/BackChannel/Transport/IBackChannelTransport.h b/Engine/Plugins/Experimental/BackChannel/Source/BackChannel/Public/BackChannel/Transport/IBackChannelTransport.h
new file mode 100644
--- /dev/null
+++ b/Engine/Plugins/Experimental/BackChannel/Source/BackChannel/Public/BackChannel/Transport/IBackChannelTransport.h
@@ -0,0 +1,35 @@
+// Copyright Epic Games, Inc. All Rights Reserved.
+
+#pragma once
+
+#include "IBackChannelSocketConnection.h"
+#include "Modules/ModuleManager.h"
+
+/**
+ *	Main module and factory interface for Backchannel connections
+ */
+class BACKCHANNEL_API IBackChannelTransport : public IModuleInterface
+{
+public:
+
+	static inline bool IsAvailable(void)
+	{
+		return Get() != nullptr;
+	}
+
+	static inline IBackChannelTransport* Get(void)
+	{
+		return FModuleManager::LoadModulePtr<IBackChannelTransport>("BackChannel");
+	}
+
+	virtual TSharedPtr<IBackChannelSocketConnection> CreateConnection(const int32 Type) = 0;
+
+public:
+
+	static const int TCP;
+
+protected:
+
+	IBackChannelTransport() {}
+	virtual ~IBackChannelTransport() {}
+};
diff --git a/Engine/Plugins/Experimental/BackChannel/Source/BackChannel/Public/BackChannel/Utils/BackChannelThreadedConnection.h b/Engine/Plugins/Experimental/BackChannel/Source/BackChannel/Public/BackChannel/Utils/BackChannelThreadedConnection.h
new file mode 100644
--- /dev/null
+++ b/Engine/Plugins/Experimental/BackChannel/Source/BackChannel/Public/BackChannel/Utils/BackChannelThreadedConnection.h
@@ -0,0 +1,41 @@
+// Copyright Epic Games, Inc. All Rights Reserved.
+
+#pragma once
+
+#include "BackChannel/Transport/IBackChannelSocketConnection.h"
+#include "HAL/Runnable.h"
+#include "HAL/ThreadSafeBool.h"
+
+class FSocket;
+
+DECLARE_DELEGATE_RetVal_OneParam(bool, FBackChannelListenerDelegate, TSharedRef<IBackChannelSocketConnection>);
+
+/**
+ * BackChannelClient implementation.
+ */
+class BACKCHANNEL_API FBackChannelThreadedListener : public FRunnable, public TSharedFromThis<FBackChannelThreadedListener>
+{
+public:
+
+	FBackChannelThreadedListener();
+	~FBackChannelThreadedListener();
+
+	void Start(TSharedRef<IBackChannelSocketConnection> InConnection, FBackChannelListenerDelegate InDelegate);
+
+	virtual void Stop() override;
+
+	bool IsRunning() const;
+
+protected:
+
+	virtual uint32 Run() override;
+
+private:
+
+	TSharedPtr<IBackChannelSocketConnection>		Connection;
+	FBackChannelListenerDelegate			Delegate;
+	
+	FThreadSafeBool							bExitRequested;
+	FThreadSafeBool							bIsRunning;
+	FCriticalSection						RunningCS;
+};
diff --git a/Engine/Plugins/Experimental/BackChannel/Source/BackChannel/Transport/IBackChannelSocketConnection.h b/Engine/Plugins/Experimental/BackChannel/Source/BackChannel/Transport/IBackChannelSocketConnection.h
deleted file mode 100644
diff --git a/Engine/Plugins/Experimental/BackChannel/Source/BackChannel/Transport/IBackChannelTransport.h b/Engine/Plugins/Experimental/BackChannel/Source/BackChannel/Transport/IBackChannelTransport.h
deleted file mode 100644
diff --git a/Engine/Plugins/Experimental/BackChannel/Source/BackChannel/Utils/BackChannelThreadedConnection.cpp b/Engine/Plugins/Experimental/BackChannel/Source/BackChannel/Utils/BackChannelThreadedConnection.cpp
deleted file mode 100644
diff --git a/Engine/Plugins/Experimental/BackChannel/Source/BackChannel/Utils/BackChannelThreadedConnection.h b/Engine/Plugins/Experimental/BackChannel/Source/BackChannel/Utils/BackChannelThreadedConnection.h
deleted file mode 100644
