BackChannel plugin: Move Protocol folder into Public/BackChannel to match rest of module

#rnx

diff --git a/Engine/Plugins/Experimental/BackChannel/Source/BackChannel/Protocol/OSC/BackChannelOSC.h b/Engine/Plugins/Experimental/BackChannel/Source/BackChannel/Protocol/OSC/BackChannelOSC.h
deleted file mode 100644
diff --git a/Engine/Plugins/Experimental/BackChannel/Source/BackChannel/Protocol/OSC/BackChannelOSCBundle.h b/Engine/Plugins/Experimental/BackChannel/Source/BackChannel/Protocol/OSC/BackChannelOSCBundle.h
deleted file mode 100644
diff --git a/Engine/Plugins/Experimental/BackChannel/Source/BackChannel/Protocol/OSC/BackChannelOSCConnection.h b/Engine/Plugins/Experimental/BackChannel/Source/BackChannel/Protocol/OSC/BackChannelOSCConnection.h
deleted file mode 100644
diff --git a/Engine/Plugins/Experimental/BackChannel/Source/BackChannel/Protocol/OSC/BackChannelOSCMessage.h b/Engine/Plugins/Experimental/BackChannel/Source/BackChannel/Protocol/OSC/BackChannelOSCMessage.h
deleted file mode 100644
diff --git a/Engine/Plugins/Experimental/BackChannel/Source/BackChannel/Protocol/OSC/BackChannelOSCPacket.h b/Engine/Plugins/Experimental/BackChannel/Source/BackChannel/Protocol/OSC/BackChannelOSCPacket.h
deleted file mode 100644
diff --git a/Engine/Plugins/Experimental/BackChannel/Source/BackChannel/Public/BackChannel/Protocol/OSC/BackChannelOSC.h b/Engine/Plugins/Experimental/BackChannel/Source/BackChannel/Public/BackChannel/Protocol/OSC/BackChannelOSC.h
new file mode 100644
--- /dev/null
+++ b/Engine/Plugins/Experimental/BackChannel/Source/BackChannel/Public/BackChannel/Protocol/OSC/BackChannelOSC.h
@@ -0,0 +1,10 @@
+// Copyright Epic Games, Inc. All Rights Reserved.
+
+#pragma once
+
+
+#include "CoreMinimal.h"
+
+#include "BackChannel/Protocol/OSC/BackChannelOSCPacket.h"
+#include "BackChannel/Protocol/OSC/BackChannelOSCMessage.h"
+#include "BackChannel/Protocol/OSC/BackChannelOSCBundle.h"
\ No newline at end of file
diff --git a/Engine/Plugins/Experimental/BackChannel/Source/BackChannel/Public/BackChannel/Protocol/OSC/BackChannelOSCBundle.h b/Engine/Plugins/Experimental/BackChannel/Source/BackChannel/Public/BackChannel/Protocol/OSC/BackChannelOSCBundle.h
new file mode 100644
--- /dev/null
+++ b/Engine/Plugins/Experimental/BackChannel/Source/BackChannel/Public/BackChannel/Protocol/OSC/BackChannelOSCBundle.h
@@ -0,0 +1,46 @@
+// Copyright Epic Games, Inc. All Rights Reserved.
+
+#pragma once
+
+#include "BackChannel/Protocol/OSC/BackChannelOSCPacket.h"
+
+
+class FBackChannelOSCBundle : public FBackChannelOSCPacket
+{
+public:
+
+	FBackChannelOSCBundle(OSCPacketMode InMode);
+	virtual ~FBackChannelOSCBundle();
+
+	FBackChannelOSCBundle(FBackChannelOSCBundle&& RHS);
+
+	FBackChannelOSCBundle& operator=(FBackChannelOSCBundle&& RHS);
+
+	virtual OSCPacketType GetType() const override { return OSCPacketType::Bundle;  }
+
+	virtual int32 GetSize() const override;
+
+	void AddElement(const void* InData, const int32 InSize);
+
+	int32 GetElementCount() const;
+
+	const TArray<uint8>& GetElement(const int32 Index) const;
+
+	virtual TArray<uint8> WriteToBuffer() const override;
+
+	virtual void WriteToBuffer(TArray<uint8>& Buffer) const override;
+
+	static TSharedPtr<FBackChannelOSCBundle> CreateFromBuffer(const void* Data, int32 DataLength);
+
+protected:
+
+protected:
+
+	typedef TArray<uint8> ElementData;
+
+	OSCPacketMode			Mode;
+	int64					TimeTag;
+	TArray<ElementData>		Elements;
+
+	static const ANSICHAR*		BundleHeader;
+};
\ No newline at end of file
diff --git a/Engine/Plugins/Experimental/BackChannel/Source/BackChannel/Public/BackChannel/Protocol/OSC/BackChannelOSCConnection.h b/Engine/Plugins/Experimental/BackChannel/Source/BackChannel/Public/BackChannel/Protocol/OSC/BackChannelOSCConnection.h
new file mode 100644
--- /dev/null
+++ b/Engine/Plugins/Experimental/BackChannel/Source/BackChannel/Public/BackChannel/Protocol/OSC/BackChannelOSCConnection.h
@@ -0,0 +1,124 @@
+// Copyright Epic Games, Inc. All Rights Reserved.
+
+#pragma once
+
+#include "CoreMinimal.h"
+#include "HAL/Runnable.h"
+#include "HAL/RunnableThread.h"
+#include "BackChannel/IBackChannelConnection.h"
+#include "BackChannel/Utils/DispatchMap.h"
+#include "HAL/ThreadSafeBool.h"
+
+class IBackChannelSocketConnection;
+class FBackChannelOSCPacket;
+
+/**
+ *	A class that wraps an existing BackChannel connection and provides an OSC-focused interface and
+ *	a background thread. Incoming messages are received on a background thread and queued until 
+ *	DispatchMessages() is called. Outgoing messages are sent immediately
+ */
+class BACKCHANNEL_API FBackChannelOSCConnection : FRunnable, public IBackChannelConnection
+{
+public:
+
+	FBackChannelOSCConnection(TSharedRef<IBackChannelSocketConnection> InConnection);
+
+	~FBackChannelOSCConnection();
+
+	/**
+	 * IBackChannelConnection implementation
+	 */
+public:
+	FString GetProtocolName() const override;
+
+	TBackChannelSharedPtr<IBackChannelPacket> CreatePacket() override;
+
+	int SendPacket(const TBackChannelSharedPtr<IBackChannelPacket>& Packet) override;
+
+	/* Bind a delegate to a message address */
+	FDelegateHandle AddRouteDelegate(FStringView Path, FBackChannelRouteDelegate::FDelegate Delegate) override;
+
+	/* Remove a delegate handle */
+	void RemoveRouteDelegate(FStringView Path, FDelegateHandle& InHandle) override;
+
+	/* Sets the send and receive buffer sizes*/
+	void SetBufferSizes(int32 DesiredSendSize, int32 DesiredReceiveSize) override;
+
+public:
+
+	bool StartReceiveThread();
+
+	// Begin public FRunnable overrides
+	virtual void Stop() override;
+	// end public FRunnable overrides
+	
+	/* Returns our connection state as determined by the underlying BackChannel connection */
+	bool IsConnected() const;
+
+	/* Returns true if running in the background */
+	bool IsThreaded() const;
+
+    /*
+        Checks for and dispatches any incoming messages. MaxTime is how long to wait if no data is ready to be read.
+        This function is thread-safe and be called from a backfround thread manually or by calling StartReceiveThread()
+    */
+	void ReceiveAndDispatchMessages(const float MaxTime = 0);
+
+	/* Send the provided OSC packet */
+	bool SendPacket(FBackChannelOSCPacket& Packet);
+
+	/* Set options for the specified message path */
+	void SetMessageOptions(const TCHAR* Path, int32 MaxQueuedMessages);
+
+	FString GetDescription();
+
+protected:
+	// Begin protected FRunnable overrides
+	virtual uint32 Run() override;
+	// End protected FRunnable overrides
+
+	bool SendPacketData(const void* Data, const int32 DataLen);
+
+	int32 GetMessageLimitForPath(const TCHAR* Path);
+
+	int32 GetMessageCountForPath(const TCHAR* Path);
+
+	void RemoveMessagesWithPath(const TCHAR* Path, const int32 Num = 0);
+
+	void ReceiveMessages(const float MaxTime = 0);
+
+	/* Dispatch all queued messages */
+	void DispatchMessages();
+
+
+protected:
+
+	TSharedPtr<IBackChannelSocketConnection>  Connection;
+
+	FBackChannelDispatchMap				DispatchMap;
+
+	TArray<TSharedPtr<FBackChannelOSCPacket>> ReceivedPackets;
+
+	TMap<FString, int32> MessageLimits;
+
+	FThreadSafeBool		ExitRequested;
+	FThreadSafeBool		IsRunning;
+
+	FCriticalSection	ReceiveMutex;
+	FCriticalSection	SendMutex;
+	FCriticalSection	PacketMutex;
+	TArray<uint8>		ReceiveBuffer;
+
+	
+	/* Time where we'll send a ping if no packets arrive to check the connection is alive*/
+	double				PingTime = 2;
+
+	/* Is the connection in an error state */
+	bool				HasErrorState = false;
+
+	/* How much data has been received this check? */
+	int32				ReceivedDataSize = 0;
+
+	/* How much data do we expect to receive next time? This is for OSC over TCP where the size of a packet is sent, then the packet*/
+	int32				ExpectedSizeOfNextPacket = 4;
+};
diff --git a/Engine/Plugins/Experimental/BackChannel/Source/BackChannel/Public/BackChannel/Protocol/OSC/BackChannelOSCMessage.h b/Engine/Plugins/Experimental/BackChannel/Source/BackChannel/Public/BackChannel/Protocol/OSC/BackChannelOSCMessage.h
new file mode 100644
--- /dev/null
+++ b/Engine/Plugins/Experimental/BackChannel/Source/BackChannel/Public/BackChannel/Protocol/OSC/BackChannelOSCMessage.h
@@ -0,0 +1,392 @@
+// Copyright Epic Games, Inc. All Rights Reserved.
+
+#pragma once
+
+#include "BackChannel/Protocol/OSC/BackChannelOSCPacket.h"
+#include "BackChannel/IBackChannelPacket.h"
+
+
+/**
+ *	Representation of an OSC message. Data can be read/written using the explicit
+ *	Read/Write functions, or the Serialize function / << operator where the behaviour
+ * is overloaded based on whether the message was created for reading or writing.
+ *
+ *	Any failed Reads() will result in the default value of the type (e.g. 0, 0.0, false, "")
+ *	being returned.
+ */
+class BACKCHANNEL_API FBackChannelOSCMessage : public FBackChannelOSCPacket, public IBackChannelPacket
+{
+public:
+
+	FBackChannelOSCMessage(OSCPacketMode InMode);
+
+	FBackChannelOSCMessage(const TCHAR* Address);
+
+	virtual ~FBackChannelOSCMessage();
+
+	/* Move constructor */
+	FBackChannelOSCMessage(FBackChannelOSCMessage&& RHS);
+
+	/* Move operator */
+	FBackChannelOSCMessage& operator=(FBackChannelOSCMessage&& RHS);
+
+	/* Return our type */
+	virtual OSCPacketType GetType() const override { return OSCPacketType::Message; }
+
+	/* Return our size (plus any necessary padding) */
+	virtual int32 GetSize() const override;
+
+	/* Returns a buffer with the contents of this message. Data is padded per OSC requirements */
+	virtual TArray<uint8> WriteToBuffer() const override;
+	
+	/* Writes this message into the provided buffer at an offset of Buffer.Num() */
+	virtual void WriteToBuffer(TArray<uint8>& Buffer) const override;
+
+	/* Helper to check our read/write status */
+	bool IsWriting() const { return Mode == OSCPacketMode::Write; }
+
+	/* Helper to check our read/write status */
+	bool IsReading() const { return Mode == OSCPacketMode::Read; }
+
+	/* Returns the address of this packet */
+	FString GetPath() const
+	{
+		return Address;
+	}
+
+	/* Return our argument tags */
+	const FString& GetTags() const
+	{
+		return TagString;
+	}
+	
+	/* Returns the number of arguments in this message */
+	int32 GetArgumentCount() const
+	{
+		return TagString.Len();
+	}
+	
+	/* Returns the type of our next argument */
+	TCHAR GetNextArgumentType() const
+	{
+		return TagString[TagIndex];
+	}
+
+	/* Return the size (plus padding) of all our arguments) */
+	const int32 GetArgumentSize() const
+	{
+		return Buffer.Num();
+	}
+
+	/* Set our destination address */
+	int SetPath(const TCHAR* Address);
+
+	/* Reset us for reading. The next argument read will be our first argument */
+	void	ResetRead();
+
+	virtual FBackChannelPacketType GetProtocolID() const { return FBackChannelPacketType('B','O','S','C'); }
+
+	virtual FString GetProtocolName() const { return TEXT("BackChannelOSC"); }
+
+	virtual bool IsWritable() const { return IsWriting(); }
+
+	virtual bool IsReadable() const { return IsReading(); }
+
+	//! Int32 read/write
+
+	/* Write an int32 into our arguments */
+	virtual int Write(const TCHAR* InName, const int32 Value) override
+	{
+		check(IsWriting());
+		int32 SwappedValue = !IsLegacyConnection() ? ByteSwap(Value) : Value;
+		WriteTagAndData(TEXT('i'), &SwappedValue, sizeof(SwappedValue));
+		return 0;
+	}
+
+	/* Write a float to our arguments */
+	virtual int Write(const TCHAR* InName, const float Value) override
+	{
+		check(IsWriting());
+		float SwappedValue = !IsLegacyConnection() ? ByteSwap(Value) : Value;
+		return WriteTagAndData(TEXT('f'), &SwappedValue, sizeof(SwappedValue));
+	}
+
+	/* Write a bool to our arguments */
+	virtual int Write(const TCHAR* InName, const bool Value) override
+	{
+		if (Value == true)
+		{
+			WriteTag(TEXT('T'));
+		}
+		else
+		{
+			WriteTag(TEXT('F'));
+		}
+
+		return 0;
+	}
+
+	/* Write a string to our arguments */
+	virtual int Write(const TCHAR* InName, const TCHAR* Value) override
+	{
+		return WriteTagAndData(TEXT('s'), TCHAR_TO_ANSI(Value), FCString::Strlen(Value) + 1);
+	}
+
+	/* Write a string to our arguments */
+	virtual int Write(const TCHAR* InName, const FString& Value) override
+	{
+		return Write(InName, *Value);
+	}
+
+	/* Write a blob of data to our arguments */
+	virtual int Write(const TCHAR* InName, const void* InBlob, int32 BlobSize) override
+	{
+		check(IsWriting());
+		if (IsLegacyConnection())
+		{
+			// with the legacy format it's just a fixed size
+			return WriteTagAndData(TEXT('b'), InBlob, BlobSize);
+		}
+		else
+		{
+			// OSC blobs contain and int with the size before the data
+			WriteTag(TEXT('b'));
+			WriteData(&BlobSize, sizeof(int32));
+			return WriteData(InBlob, BlobSize);
+		}
+	}
+
+	/*
+		Read data from the message into a TArray. It must have been serialized by the Read form for TArray(!).
+		Note - data will be appended to the array.
+	*/
+	virtual int Write(const TCHAR* InName, const TArrayView<const uint8> InValues) override
+	{
+		return Write<uint8>(InName, InValues);
+	}
+
+	/* Write a TArrayView into the message */
+	template<typename T>
+	int Write(const TCHAR* InName, const TArrayView<const T> InValues)
+	{
+		if (IsLegacyConnection())
+		{
+			// contrary to the OSC spec, with the legacy format the size of a blob has an explicit tag for rather than simply being the first
+			// four bytes of the blob data. See also Reading of arrays
+			Write(TEXT("Size"), InValues.Num());
+			int32 BlobSize = InValues.Num() * sizeof(T);
+			return WriteTagAndData(TEXT('b'), InValues.GetData(), BlobSize);
+		}
+		else
+		{
+			return Write(InName, InValues.GetData(), InValues.Num());
+		}
+	}
+
+	/* Read an int32 from our arguments */
+	virtual int Read(const TCHAR* InName, int32& Value) override
+	{
+		check(IsReading());
+		int Err = ReadTagAndData(TEXT('i'), &Value, sizeof(Value));
+
+		if (!Err && !IsLegacyConnection())
+		{
+			Value = ByteSwap(Value);
+		}
+
+		return Err;
+	}	
+
+	/* Read a float from our arguments */
+	virtual int Read(const TCHAR* InName, float& OutValue) override
+	{
+		check(IsReading());
+		int Err = ReadTagAndData(TEXT('f'), &OutValue, sizeof(OutValue));
+
+		if (!Err && !IsLegacyConnection())
+		{
+			OutValue = ByteSwap(OutValue);
+		}
+
+		return Err;
+	}
+
+	/* Read a bool from our arguments */
+	virtual int Read(const TCHAR* InName, bool& Value) override
+	{
+		// try to read both tagd and see which was present
+		TCHAR TrueTag = ReadTag(TEXT('T'), true);
+
+		if (TrueTag != 0)
+		{
+			Value = true;
+			return 0;
+		}
+
+		TCHAR FalseTag = ReadTag(TEXT('F'), true);
+
+		if (FalseTag != 0)
+		{
+			Value = false;
+			return 0;
+		}
+
+		return 1;
+	}
+
+	/* Read a string from our arguments.  */
+	virtual int Read(const TCHAR* InName, FString& OutValue) override;
+
+	//! Raw data blobs
+
+	/* Read a blob of data from our arguments */
+	virtual int Read(const TCHAR* InName, void* InBlob, int32 MaxBlobSize, int32& OutBlobSize) override;	
+
+	/* 
+		Read data from the message into a TArray. It must have been serialized by the Read form for TArray(!). 
+		Note - data will be appended to the array.
+	*/
+	virtual int Read(const TCHAR* InName, TArray<uint8>& DataArray) override
+	{
+		return Read<uint8>(InName, DataArray);
+	}
+
+	/*
+	 *	Read a TArray of type T from our arguments. This is a helper that reads an int
+	*	for the size, then allocated and reads a blob of sizeof(t) * NumItems
+	 */
+	template<typename T>
+	int Read(const TCHAR* InName, TArray<T>& DataArray)
+	{
+		int32 BlobSize = 0;
+		int32 NumNewElements = 0;
+
+		if (IsLegacyConnection())
+		{
+			// contrary to the OSC spec, with the legacy format the size of a blob has an explicit tag for rather than simply being the first
+			// four bytes of the blob data. See also reading of arrays
+			Read(TEXT("Size"), NumNewElements);
+			BlobSize = NumNewElements * sizeof(T);
+
+			// Add enough space
+			int NumCurrentElements = DataArray.Num();
+			DataArray.AddUninitialized(NumNewElements);
+
+			// read the tag and data directly, don't use the blob read
+			return ReadTagAndData(TEXT('b'), DataArray.GetData() + NumCurrentElements, BlobSize);
+		}
+		else
+		{
+			// Read the size of the blob and covert to element count
+			Read(InName, nullptr, 0, BlobSize);
+			NumNewElements = BlobSize / sizeof(T);
+
+			// Add enough space
+			int NumCurrentElements = DataArray.Num();
+			DataArray.AddUninitialized(NumNewElements);
+
+			return Read(InName, DataArray.GetData() + NumCurrentElements, BlobSize, BlobSize);
+		}
+	}
+	
+
+	/* Serialize helper that will read/write based on the open mode of this message */
+	template<typename T>
+	void Serialize(const TCHAR* Name, T& Value)
+	{
+		if (IsWriting())
+		{
+			Write(Name, Value);
+		}
+		else
+		{
+			Read(Name, Value);
+		}
+	}
+
+	/* Serialize helper that will read/write based on the open mode of this message */
+	void Serialize(const TCHAR* Name, void* InBlob, int32 BlobSize)
+	{
+		if (IsReading())
+		{
+			int32 OutSize(0);
+			Read(Name, InBlob, BlobSize, OutSize);
+		}
+		else
+		{
+			Write(Name, InBlob, BlobSize);
+		}
+	}
+
+	static int32 RoundedArgumentSize(int32 ArgSize)
+	{
+		return ((ArgSize + 3) / 4) * 4;
+	}
+
+	static TSharedPtr<FBackChannelOSCMessage> CreateFromBuffer(const void* Data, int32 DataLength);
+
+	static void SetLegacyMode(const bool bEnable) { bIsLegacyConnection  = bEnable; }
+	static bool IsLegacyConnection() { return bIsLegacyConnection; }
+
+protected:
+
+	/* Return our sizes for each item (plus any necessary padding) */
+	void GetComponentSizes(int32& OutAddressSize, int32& OutTagSize, int32& OutBufferSize) const;
+
+	int Serialize(const TCHAR Code, void* InData, int32 InSize);
+
+	TCHAR ReadTag(const TCHAR ExpectedTag, bool SuppressWarning=false);
+	int ReadData(void* InData, int32 InSize);
+	int ReadTagAndData(const TCHAR ExpectedTag, void* InData, int32 InSize);
+
+	int WriteTag(const TCHAR Code);
+	int WriteTagAndData(const TCHAR Code, const void* InData, int32 InSize);
+	int WriteData(const void* InData, int32 InSize);
+
+protected:
+
+	OSCPacketMode		Mode;
+	FString				Address;
+	FString				TagString;
+	int					TagIndex;
+	int					BufferIndex;
+	TArray<uint8>		Buffer;
+
+	static bool bIsLegacyConnection;
+};
+
+/*
+BACKCHANNEL_API FBackChannelOSCMessage& operator << (FBackChannelOSCMessage& Msg, int32& Value);
+
+BACKCHANNEL_API FBackChannelOSCMessage& operator << (FBackChannelOSCMessage& Msg, float& Value);
+
+BACKCHANNEL_API FBackChannelOSCMessage& operator << (FBackChannelOSCMessage& Msg, bool& Value);
+
+BACKCHANNEL_API FBackChannelOSCMessage& operator << (FBackChannelOSCMessage& Msg, TCHAR& Value);
+
+BACKCHANNEL_API FBackChannelOSCMessage& operator << (FBackChannelOSCMessage& Msg, FString& Value);
+
+template <typename T>
+FBackChannelOSCMessage& operator << (FBackChannelOSCMessage& Msg, TArray<T>& Value)
+{
+	if (Msg.IsWriting())
+	{
+		Msg.Write(nullptr, Value);
+	}
+	else
+	{
+		Msg.Read(nullptr, Value);
+	}
+
+	return Msg;
+}
+
+
+template <typename T>
+FBackChannelOSCMessage& SerializeOut(FBackChannelOSCMessage& Msg, const T& Value)
+{
+	T Tmp = Value;
+	Msg << Tmp;
+	return Msg;
+}
+*/
\ No newline at end of file
diff --git a/Engine/Plugins/Experimental/BackChannel/Source/BackChannel/Public/BackChannel/Protocol/OSC/BackChannelOSCPacket.h b/Engine/Plugins/Experimental/BackChannel/Source/BackChannel/Public/BackChannel/Protocol/OSC/BackChannelOSCPacket.h
new file mode 100644
--- /dev/null
+++ b/Engine/Plugins/Experimental/BackChannel/Source/BackChannel/Public/BackChannel/Protocol/OSC/BackChannelOSCPacket.h
@@ -0,0 +1,55 @@
+// Copyright Epic Games, Inc. All Rights Reserved.
+
+#pragma once
+
+#include "CoreMinimal.h"
+
+enum class OSCPacketMode
+{
+	Read,
+	Write
+};
+
+enum class OSCPacketType
+{
+	Invalid,
+	Message,
+	Bundle
+};
+
+/**
+ *	Base class for OSC messages
+ */
+class BACKCHANNEL_API FBackChannelOSCPacket
+{
+public:
+
+	virtual ~FBackChannelOSCPacket() {}
+
+	/* Return the total size in bytes of this message */
+	virtual int32	GetSize() const = 0;
+
+	/* Return the type of this packet */
+	virtual OSCPacketType GetType() const = 0;
+	
+	/* Return a buffer with a copy of the contents of this packet */
+	virtual TArray<uint8> WriteToBuffer() const =0;
+
+	/* Write the contents of this packet into the specified buffer starting at an offset of Buffer.Num()  */
+	virtual void WriteToBuffer(TArray<uint8>& Buffer) const = 0;
+
+	/* Static helper to determine the type of packet (if any) at the specified address */
+	static OSCPacketType GetType(const void* Data, int32 DataLength);
+
+	/* Helper to construct a packet (and in the case of a bundle, any sub-packets) from the specified buffer */
+	static TSharedPtr<FBackChannelOSCPacket> CreateFromBuffer(const void* Data, int32 DataLength);
+};
+
+class FBackChannelOSCNullPacket : public FBackChannelOSCPacket
+{
+public:
+	virtual int32 GetSize() const override { return 0; }
+
+	virtual OSCPacketType GetType() const override { return OSCPacketType::Invalid;  }
+};
+
