Fix Win64 Shipping -DisableUnity -NoPCH issues due to incorrect stats header

#rnx

diff --git a/Engine/Plugins/2D/Paper2D/Source/Paper2D/Public/Paper2DModule.h b/Engine/Plugins/2D/Paper2D/Source/Paper2D/Public/Paper2DModule.h
--- a/Engine/Plugins/2D/Paper2D/Source/Paper2D/Public/Paper2DModule.h
+++ b/Engine/Plugins/2D/Paper2D/Source/Paper2D/Public/Paper2DModule.h
@@ -4,7 +4,7 @@
 
 #include "Math/Vector.h"
 #include "Modules/ModuleInterface.h"
-#include "Stats/Stats2.h"
+#include "Stats/Stats.h"
 
 enum class EStatFlags : uint8;
 
diff --git a/Engine/Plugins/Cameras/GameplayCameras/Source/GameplayCameras/Public/GameplayCamerasModule.h b/Engine/Plugins/Cameras/GameplayCameras/Source/GameplayCameras/Public/GameplayCamerasModule.h
--- a/Engine/Plugins/Cameras/GameplayCameras/Source/GameplayCameras/Public/GameplayCamerasModule.h
+++ b/Engine/Plugins/Cameras/GameplayCameras/Source/GameplayCameras/Public/GameplayCamerasModule.h
@@ -3,7 +3,7 @@
 #pragma once
 
 #include "Modules/ModuleInterface.h"
-#include "Stats/Stats2.h"
+#include "Stats/Stats.h"
 
 DECLARE_STATS_GROUP(TEXT("Camera Animation Evaluation"), STATGROUP_CameraAnimation, STATCAT_Advanced)
 
diff --git a/Engine/Plugins/Experimental/BackChannel/Source/BackChannel/Private/Transport/BackChannelConnection.cpp b/Engine/Plugins/Experimental/BackChannel/Source/BackChannel/Private/Transport/BackChannelConnection.cpp
--- a/Engine/Plugins/Experimental/BackChannel/Source/BackChannel/Private/Transport/BackChannelConnection.cpp
+++ b/Engine/Plugins/Experimental/BackChannel/Source/BackChannel/Private/Transport/BackChannelConnection.cpp
@@ -7,7 +7,7 @@
 #include "Misc/ConfigCacheIni.h"
 #include "Misc/OutputDeviceRedirector.h"
 #include "Sockets.h"
-#include "Stats/Stats2.h"
+#include "Stats/Stats.h"
 #include "Stats/StatsTrace.h"
 
 DECLARE_STATS_GROUP(TEXT("BackChannel"), STATGROUP_BackChannel, STATCAT_Advanced);
diff --git a/Engine/Plugins/Experimental/Web/WebAPI/Source/WebAPI/Private/WebAPIStats.h b/Engine/Plugins/Experimental/Web/WebAPI/Source/WebAPI/Private/WebAPIStats.h
--- a/Engine/Plugins/Experimental/Web/WebAPI/Source/WebAPI/Private/WebAPIStats.h
+++ b/Engine/Plugins/Experimental/Web/WebAPI/Source/WebAPI/Private/WebAPIStats.h
@@ -3,6 +3,6 @@
 #pragma once
 
 
-#include "Stats/Stats2.h"
+#include "Stats/Stats.h"
 
 DECLARE_STATS_GROUP(TEXT("WebAPI"), STATGROUP_WebAPI, STATCAT_Advanced);
diff --git a/Engine/Plugins/MovieScene/MovieRenderPipeline/Source/MovieRenderPipelineCore/Public/MovieRenderPipelineCoreModule.h b/Engine/Plugins/MovieScene/MovieRenderPipeline/Source/MovieRenderPipelineCore/Public/MovieRenderPipelineCoreModule.h
--- a/Engine/Plugins/MovieScene/MovieRenderPipeline/Source/MovieRenderPipelineCore/Public/MovieRenderPipelineCoreModule.h
+++ b/Engine/Plugins/MovieScene/MovieRenderPipeline/Source/MovieRenderPipelineCore/Public/MovieRenderPipelineCoreModule.h
@@ -3,7 +3,7 @@
 #pragma once
 
 #include "CoreMinimal.h"
-#include "Stats/Stats2.h"
+#include "Stats/Stats.h"
 #include "Modules/ModuleInterface.h"
 #include "Features/IModularFeature.h"
 
diff --git a/Engine/Plugins/Protocols/MQTT/Source/MQTTCore/Private/MQTTCoreStats.h b/Engine/Plugins/Protocols/MQTT/Source/MQTTCore/Private/MQTTCoreStats.h
--- a/Engine/Plugins/Protocols/MQTT/Source/MQTTCore/Private/MQTTCoreStats.h
+++ b/Engine/Plugins/Protocols/MQTT/Source/MQTTCore/Private/MQTTCoreStats.h
@@ -3,6 +3,6 @@
 #pragma once
 
 
-#include "Stats/Stats2.h"
+#include "Stats/Stats.h"
 
 DECLARE_STATS_GROUP(TEXT("MQTT"), STATGROUP_MQTT, STATCAT_Advanced);
diff --git a/Engine/Plugins/Runtime/ModelViewViewModel/Source/ModelViewViewModel/Private/ModelViewViewModelModule.h b/Engine/Plugins/Runtime/ModelViewViewModel/Source/ModelViewViewModel/Private/ModelViewViewModelModule.h
--- a/Engine/Plugins/Runtime/ModelViewViewModel/Source/ModelViewViewModel/Private/ModelViewViewModelModule.h
+++ b/Engine/Plugins/Runtime/ModelViewViewModel/Source/ModelViewViewModel/Private/ModelViewViewModelModule.h
@@ -3,7 +3,7 @@
 #pragma once
 
 #include "Modules/ModuleInterface.h"
-#include "Stats/Stats2.h"
+#include "Stats/Stats.h"
 
 class IConsoleVariable;
 
diff --git a/Engine/Plugins/Runtime/nDisplay/Source/DisplayCluster/Private/Misc/DisplayClusterLog.h b/Engine/Plugins/Runtime/nDisplay/Source/DisplayCluster/Private/Misc/DisplayClusterLog.h
--- a/Engine/Plugins/Runtime/nDisplay/Source/DisplayCluster/Private/Misc/DisplayClusterLog.h
+++ b/Engine/Plugins/Runtime/nDisplay/Source/DisplayCluster/Private/Misc/DisplayClusterLog.h
@@ -3,7 +3,7 @@
 #pragma once
 
 #include "CoreMinimal.h"
-#include "Stats/Stats2.h"
+#include "Stats/Stats.h"
 
 // Module-wide log categories
 DECLARE_LOG_CATEGORY_EXTERN(LogDisplayClusterEngine,     Log, All);
