Replaced EAutomationTestFlags::ApplicationContextMask with EAutomationTestFlags_ApplicationContextMask.

diff --git a/Engine/Plugins/Animation/LiveLink/Source/LiveLinkEditor/Private/LiveLinkTest.cpp b/Engine/Plugins/Animation/LiveLink/Source/LiveLinkEditor/Private/LiveLinkTest.cpp
--- a/Engine/Plugins/Animation/LiveLink/Source/LiveLinkEditor/Private/LiveLinkTest.cpp
+++ b/Engine/Plugins/Animation/LiveLink/Source/LiveLinkEditor/Private/LiveLinkTest.cpp
@@ -8,7 +8,7 @@
 
 #if WITH_DEV_AUTOMATION_TESTS
 
-IMPLEMENT_SIMPLE_AUTOMATION_TEST(FLiveLinkInterpolationTest, "System.Engine.Animation.LiveLink.Interpolation", EAutomationTestFlags::ApplicationContextMask | EAutomationTestFlags::EngineFilter)
+IMPLEMENT_SIMPLE_AUTOMATION_TEST(FLiveLinkInterpolationTest, "System.Engine.Animation.LiveLink.Interpolation", EAutomationTestFlags_ApplicationContextMask | EAutomationTestFlags::EngineFilter)
 
 /** */
 bool FLiveLinkInterpolationTest::RunTest(const FString& Parameters)
diff --git a/Engine/Plugins/Animation/MLDeformer/VertexDeltaModel/Source/VertexDeltaModel/Private/Tests/VertexDeltaModelTest.cpp b/Engine/Plugins/Animation/MLDeformer/VertexDeltaModel/Source/VertexDeltaModel/Private/Tests/VertexDeltaModelTest.cpp
--- a/Engine/Plugins/Animation/MLDeformer/VertexDeltaModel/Source/VertexDeltaModel/Private/Tests/VertexDeltaModelTest.cpp
+++ b/Engine/Plugins/Animation/MLDeformer/VertexDeltaModel/Source/VertexDeltaModel/Private/Tests/VertexDeltaModelTest.cpp
@@ -22,7 +22,7 @@
 namespace UE::VertexDeltaModelTests
 {
 /*
-	IMPLEMENT_SIMPLE_AUTOMATION_TEST(FVertexDeltaModelMainTest, "MLDeformer.VertexDeltaModel.MainTest", EAutomationTestFlags::ApplicationContextMask | EAutomationTestFlags::EngineFilter);
+	IMPLEMENT_SIMPLE_AUTOMATION_TEST(FVertexDeltaModelMainTest, "MLDeformer.VertexDeltaModel.MainTest", EAutomationTestFlags_ApplicationContextMask | EAutomationTestFlags::EngineFilter);
 	bool FVertexDeltaModelMainTest::RunTest(const FString& Parameters)
 	{
 		using namespace UE::MLDeformer;
diff --git a/Engine/Plugins/Developer/TraceSourceFiltering/Source/SourceFilteringTrace/Public/TraceSourceFilteringTestBase.h b/Engine/Plugins/Developer/TraceSourceFiltering/Source/SourceFilteringTrace/Public/TraceSourceFilteringTestBase.h
--- a/Engine/Plugins/Developer/TraceSourceFiltering/Source/SourceFilteringTrace/Public/TraceSourceFilteringTestBase.h
+++ b/Engine/Plugins/Developer/TraceSourceFiltering/Source/SourceFilteringTrace/Public/TraceSourceFilteringTestBase.h
@@ -23,7 +23,7 @@ public:
 	virtual ~FTraceSourceFilteringTestBase() {};
 
 	/** Begin FAutomationTestBase overrides */
-	virtual EAutomationTestFlags GetTestFlags() const override { return EAutomationTestFlags::ApplicationContextMask | EAutomationTestFlags::EngineFilter; }
+	virtual EAutomationTestFlags GetTestFlags() const override { return EAutomationTestFlags_ApplicationContextMask | EAutomationTestFlags::EngineFilter; }
 	virtual bool IsStressTest() const {	return false; }
 	virtual uint32 GetRequiredDeviceNum() const override { return 1; }
 	virtual bool RunTest(const FString& Parameters) override;
diff --git a/Engine/Plugins/Experimental/Avalanche/Source/Avalanche/Private/Tests/AvaGizmoComponent.Spec.cpp b/Engine/Plugins/Experimental/Avalanche/Source/Avalanche/Private/Tests/AvaGizmoComponent.Spec.cpp
--- a/Engine/Plugins/Experimental/Avalanche/Source/Avalanche/Private/Tests/AvaGizmoComponent.Spec.cpp
+++ b/Engine/Plugins/Experimental/Avalanche/Source/Avalanche/Private/Tests/AvaGizmoComponent.Spec.cpp
@@ -1,4 +1,4 @@
-﻿// Copyright Epic Games, Inc. All Rights Reserved.
+// Copyright Epic Games, Inc. All Rights Reserved.
 
 #include "Components/StaticMeshComponent.h"
 #include "Engine/StaticMeshActor.h"
@@ -7,7 +7,7 @@
 #include "Tests/Framework/AvaTestUtils.h"
 
 BEGIN_DEFINE_SPEC(FAvalancheGizmoComponentSpec, "Avalanche.GizmoComponent",
-                  EAutomationTestFlags::ProductFilter | EAutomationTestFlags::ApplicationContextMask)
+                  EAutomationTestFlags::ProductFilter | EAutomationTestFlags_ApplicationContextMask)
 
 	TSharedPtr<FAvaTestUtils> TestUtils = MakeShared<FAvaTestUtils>();
 
diff --git a/Engine/Plugins/Experimental/Avalanche/Source/AvalancheEditor/Private/Tests/AvaEditorPalettes.spec.cpp b/Engine/Plugins/Experimental/Avalanche/Source/AvalancheEditor/Private/Tests/AvaEditorPalettes.spec.cpp
--- a/Engine/Plugins/Experimental/Avalanche/Source/AvalancheEditor/Private/Tests/AvaEditorPalettes.spec.cpp
+++ b/Engine/Plugins/Experimental/Avalanche/Source/AvalancheEditor/Private/Tests/AvaEditorPalettes.spec.cpp
@@ -8,7 +8,7 @@
 #include "Tests/Framework/AvaEditorTestUtils.h"
 #include "Tests/Framework/AvaTestDynamicMeshActor.h"
 
-BEGIN_DEFINE_SPEC(AvalancheEditorAligmentPalettes, "Avalanche.Editor.AlignmentPalettes", EAutomationTestFlags::ProductFilter | EAutomationTestFlags::ApplicationContextMask)
+BEGIN_DEFINE_SPEC(AvalancheEditorAligmentPalettes, "Avalanche.Editor.AlignmentPalettes", EAutomationTestFlags::ProductFilter | EAutomationTestFlags_ApplicationContextMask)
 
 	TArray<AActor*> Actors;
 	int32 ActorCount = 3;
diff --git a/Engine/Plugins/Experimental/Avalanche/Source/AvalancheModifiers/Private/Tests/AvaModifiersAlignBetween.spec.cpp b/Engine/Plugins/Experimental/Avalanche/Source/AvalancheModifiers/Private/Tests/AvaModifiersAlignBetween.spec.cpp
--- a/Engine/Plugins/Experimental/Avalanche/Source/AvalancheModifiers/Private/Tests/AvaModifiersAlignBetween.spec.cpp
+++ b/Engine/Plugins/Experimental/Avalanche/Source/AvalancheModifiers/Private/Tests/AvaModifiersAlignBetween.spec.cpp
@@ -9,7 +9,7 @@
 #include "Tests/Framework/AvaTestUtils.h"
 
 BEGIN_DEFINE_SPEC(AvalancheModifiersAlignBetween, "Avalanche.Modifiers.AlignBetween",
-                  EAutomationTestFlags::ProductFilter | EAutomationTestFlags::ApplicationContextMask)
+                  EAutomationTestFlags::ProductFilter | EAutomationTestFlags_ApplicationContextMask)
 
 	AAvaTestDynamicMeshActor* ModifiedMeshActor;
 	AAvaTestDynamicMeshActor* ReferenceOneMeshActor;
diff --git a/Engine/Plugins/Experimental/Avalanche/Source/AvalancheModifiers/Private/Tests/AvaModifiersAutoFollow.spec.cpp b/Engine/Plugins/Experimental/Avalanche/Source/AvalancheModifiers/Private/Tests/AvaModifiersAutoFollow.spec.cpp
--- a/Engine/Plugins/Experimental/Avalanche/Source/AvalancheModifiers/Private/Tests/AvaModifiersAutoFollow.spec.cpp
+++ b/Engine/Plugins/Experimental/Avalanche/Source/AvalancheModifiers/Private/Tests/AvaModifiersAutoFollow.spec.cpp
@@ -10,7 +10,7 @@
 #include "Tests/Framework/AvaTestUtils.h"
 
 BEGIN_DEFINE_SPEC(AvalancheModifiersAutoFollowTests, "Avalanche.Modifiers.Autofollow",
-                  EAutomationTestFlags::ProductFilter | EAutomationTestFlags::ApplicationContextMask)
+                  EAutomationTestFlags::ProductFilter | EAutomationTestFlags_ApplicationContextMask)
 
 	FName ModifierAutoFollowName;
 
diff --git a/Engine/Plugins/Experimental/Avalanche/Source/AvalancheModifiers/Private/Tests/AvaModifiersAutoSize.spec.cpp b/Engine/Plugins/Experimental/Avalanche/Source/AvalancheModifiers/Private/Tests/AvaModifiersAutoSize.spec.cpp
--- a/Engine/Plugins/Experimental/Avalanche/Source/AvalancheModifiers/Private/Tests/AvaModifiersAutoSize.spec.cpp
+++ b/Engine/Plugins/Experimental/Avalanche/Source/AvalancheModifiers/Private/Tests/AvaModifiersAutoSize.spec.cpp
@@ -12,7 +12,7 @@
 #if WITH_AUTOMATION_TESTS
 
 BEGIN_DEFINE_SPEC(AvalancheModifiersAutoSize, "Avalanche.Modifiers.AutoSize",
-                  EAutomationTestFlags::ProductFilter | EAutomationTestFlags::ApplicationContextMask)
+                  EAutomationTestFlags::ProductFilter | EAutomationTestFlags_ApplicationContextMask)
 
 	AAvaTestDynamicMeshActor* ReferenceActor;
 	AAvaTestDynamicMeshActor* ModifiedActor;
diff --git a/Engine/Plugins/Experimental/Avalanche/Source/AvalancheModifiers/Private/Tests/AvaModifiersCore.spec.cpp b/Engine/Plugins/Experimental/Avalanche/Source/AvalancheModifiers/Private/Tests/AvaModifiersCore.spec.cpp
--- a/Engine/Plugins/Experimental/Avalanche/Source/AvalancheModifiers/Private/Tests/AvaModifiersCore.spec.cpp
+++ b/Engine/Plugins/Experimental/Avalanche/Source/AvalancheModifiers/Private/Tests/AvaModifiersCore.spec.cpp
@@ -9,7 +9,7 @@
 #include "Subsystems/ActorModifierCoreSubsystem.h"
 
 BEGIN_DEFINE_SPEC(AvalancheModifiersCoreTests, "Avalanche.Modifiers.Core",
-                  EAutomationTestFlags::ProductFilter | EAutomationTestFlags::ApplicationContextMask)
+                  EAutomationTestFlags::ProductFilter | EAutomationTestFlags_ApplicationContextMask)
 
 	FAvaTestData TestData;
 	AAvaTestDynamicMeshActor* DynamicMeshActor;
diff --git a/Engine/Plugins/Experimental/Avalanche/Source/AvalancheModifiers/Private/Tests/AvaModifiersGridArrange.spec.cpp b/Engine/Plugins/Experimental/Avalanche/Source/AvalancheModifiers/Private/Tests/AvaModifiersGridArrange.spec.cpp
--- a/Engine/Plugins/Experimental/Avalanche/Source/AvalancheModifiers/Private/Tests/AvaModifiersGridArrange.spec.cpp
+++ b/Engine/Plugins/Experimental/Avalanche/Source/AvalancheModifiers/Private/Tests/AvaModifiersGridArrange.spec.cpp
@@ -9,7 +9,7 @@
 #include "Tests/Framework/AvaTestUtils.h"
 
 BEGIN_DEFINE_SPEC(AvalancheModifiersGridArrange, "Avalanche.Modifiers.GridArrange",
-                  EAutomationTestFlags::ProductFilter | EAutomationTestFlags::ApplicationContextMask)
+                  EAutomationTestFlags::ProductFilter | EAutomationTestFlags_ApplicationContextMask)
 
 	AAvaTestDynamicMeshActor* ParentMeshActor;
 	AAvaTestDynamicMeshActor* ChildOneMeshActor;
diff --git a/Engine/Plugins/Experimental/Avalanche/Source/AvalancheModifiers/Private/Tests/AvaModifiersJustify.spec.cpp b/Engine/Plugins/Experimental/Avalanche/Source/AvalancheModifiers/Private/Tests/AvaModifiersJustify.spec.cpp
--- a/Engine/Plugins/Experimental/Avalanche/Source/AvalancheModifiers/Private/Tests/AvaModifiersJustify.spec.cpp
+++ b/Engine/Plugins/Experimental/Avalanche/Source/AvalancheModifiers/Private/Tests/AvaModifiersJustify.spec.cpp
@@ -9,7 +9,7 @@
 #include "Tests/Framework/AvaTestUtils.h"
 
 BEGIN_DEFINE_SPEC(AvalancheModifiersJustify, "Avalanche.Modifiers.Justify",
-                  EAutomationTestFlags::ProductFilter | EAutomationTestFlags::ApplicationContextMask)
+                  EAutomationTestFlags::ProductFilter | EAutomationTestFlags_ApplicationContextMask)
 
 	AAvaTestDynamicMeshActor* ParentActor;
 	AAvaTestDynamicMeshActor* ChildOneActor;
diff --git a/Engine/Plugins/Experimental/Avalanche/Source/AvalancheModifiers/Private/Tests/AvaModifiersLookAt.spec.cpp b/Engine/Plugins/Experimental/Avalanche/Source/AvalancheModifiers/Private/Tests/AvaModifiersLookAt.spec.cpp
--- a/Engine/Plugins/Experimental/Avalanche/Source/AvalancheModifiers/Private/Tests/AvaModifiersLookAt.spec.cpp
+++ b/Engine/Plugins/Experimental/Avalanche/Source/AvalancheModifiers/Private/Tests/AvaModifiersLookAt.spec.cpp
@@ -10,7 +10,7 @@
 #include "Tests/Framework/AvaTestUtils.h"
 
 BEGIN_DEFINE_SPEC(AvalancheModifiersLookAt, "Avalanche.Modifiers.LookAt",
-                  EAutomationTestFlags::ProductFilter | EAutomationTestFlags::ApplicationContextMask)
+                  EAutomationTestFlags::ProductFilter | EAutomationTestFlags_ApplicationContextMask)
 
 	AAvaTestDynamicMeshActor* ModifiedActor;
 	AAvaTestDynamicMeshActor* ReferenceActor;
diff --git a/Engine/Plugins/Experimental/Avalanche/Source/AvalancheModifiers/Private/Tests/AvaModifiersRadialArrange.cpp b/Engine/Plugins/Experimental/Avalanche/Source/AvalancheModifiers/Private/Tests/AvaModifiersRadialArrange.cpp
--- a/Engine/Plugins/Experimental/Avalanche/Source/AvalancheModifiers/Private/Tests/AvaModifiersRadialArrange.cpp
+++ b/Engine/Plugins/Experimental/Avalanche/Source/AvalancheModifiers/Private/Tests/AvaModifiersRadialArrange.cpp
@@ -9,7 +9,7 @@
 #include "Tests/Framework/AvaTestUtils.h"
 
 BEGIN_DEFINE_SPEC(AvalancheModifiersRadialArrange, "Avalanche.Modifiers.RadialArrange",
-                  EAutomationTestFlags::ProductFilter | EAutomationTestFlags::ApplicationContextMask)
+                  EAutomationTestFlags::ProductFilter | EAutomationTestFlags_ApplicationContextMask)
 
 	AAvaTestDynamicMeshActor* ParentActor;
 	TArray<AAvaTestDynamicMeshActor*> ChildActors;
diff --git a/Engine/Plugins/Experimental/Avalanche/Source/AvalancheModifiers/Private/Tests/AvaModifiersTranslucentPriority.spec.cpp b/Engine/Plugins/Experimental/Avalanche/Source/AvalancheModifiers/Private/Tests/AvaModifiersTranslucentPriority.spec.cpp
--- a/Engine/Plugins/Experimental/Avalanche/Source/AvalancheModifiers/Private/Tests/AvaModifiersTranslucentPriority.spec.cpp
+++ b/Engine/Plugins/Experimental/Avalanche/Source/AvalancheModifiers/Private/Tests/AvaModifiersTranslucentPriority.spec.cpp
@@ -10,7 +10,7 @@
 #include "Tests/Framework/AvaTestUtils.h"
 
 BEGIN_DEFINE_SPEC(AvalancheModifiersTranslucentPriority, "Avalanche.Modifiers.TranslucentPriority",
-                  EAutomationTestFlags::ProductFilter | EAutomationTestFlags::ApplicationContextMask)
+                  EAutomationTestFlags::ProductFilter | EAutomationTestFlags_ApplicationContextMask)
 	AAvaTestDynamicMeshActor* ModifiedActor;
 	TArray<UPrimitiveComponent*> ModifiedActorComponents;
 	int32 ModifierSortPriority;
diff --git a/Engine/Plugins/Experimental/Avalanche/Source/AvalancheModifiers/Private/Tests/AvaModifiersVisibility.spec.cpp b/Engine/Plugins/Experimental/Avalanche/Source/AvalancheModifiers/Private/Tests/AvaModifiersVisibility.spec.cpp
--- a/Engine/Plugins/Experimental/Avalanche/Source/AvalancheModifiers/Private/Tests/AvaModifiersVisibility.spec.cpp
+++ b/Engine/Plugins/Experimental/Avalanche/Source/AvalancheModifiers/Private/Tests/AvaModifiersVisibility.spec.cpp
@@ -9,7 +9,7 @@
 #include "Tests/Framework/AvaTestUtils.h"
 
 BEGIN_DEFINE_SPEC(AvalancheModifiersVisibility, "Avalanche.Modifiers.Visibility",
-                  EAutomationTestFlags::ProductFilter | EAutomationTestFlags::ApplicationContextMask)
+                  EAutomationTestFlags::ProductFilter | EAutomationTestFlags_ApplicationContextMask)
 
 	AAvaTestDynamicMeshActor* ParentActor;
 	TArray<AAvaTestDynamicMeshActor*> ChildrenActors;
diff --git a/Engine/Plugins/Experimental/Avalanche/Source/AvalancheOutliner/Private/Tests/AvaOutliner.spec.cpp b/Engine/Plugins/Experimental/Avalanche/Source/AvalancheOutliner/Private/Tests/AvaOutliner.spec.cpp
--- a/Engine/Plugins/Experimental/Avalanche/Source/AvalancheOutliner/Private/Tests/AvaOutliner.spec.cpp
+++ b/Engine/Plugins/Experimental/Avalanche/Source/AvalancheOutliner/Private/Tests/AvaOutliner.spec.cpp
@@ -14,7 +14,7 @@
 BEGIN_DEFINE_SPEC(FAvaOutlinerSpec, "Avalanche.Outliner"
 	, EAutomationTestFlags::EngineFilter
 	| EAutomationTestFlags::EditorContext
-	| EAutomationTestFlags::ApplicationContextMask)
+	| EAutomationTestFlags_ApplicationContextMask)
 
 	// Types of common situations that are used across multiple tests
 	enum class EWhenType
diff --git a/Engine/Plugins/Experimental/BackChannel/Source/BackChannel/Private/Protocol/TestBackChannelOSC.cpp b/Engine/Plugins/Experimental/BackChannel/Source/BackChannel/Private/Protocol/TestBackChannelOSC.cpp
--- a/Engine/Plugins/Experimental/BackChannel/Source/BackChannel/Private/Protocol/TestBackChannelOSC.cpp
+++ b/Engine/Plugins/Experimental/BackChannel/Source/BackChannel/Private/Protocol/TestBackChannelOSC.cpp
@@ -17,11 +17,11 @@ public:
 
 };
 
-IMPLEMENT_CUSTOM_SIMPLE_AUTOMATION_TEST(FBackChannelTestOSCMessage, FBackChannelTestOSCBase, "BackChannel.TestOSCMessage", EAutomationTestFlags::ApplicationContextMask | EAutomationTestFlags::EngineFilter)
+IMPLEMENT_CUSTOM_SIMPLE_AUTOMATION_TEST(FBackChannelTestOSCMessage, FBackChannelTestOSCBase, "BackChannel.TestOSCMessage", EAutomationTestFlags_ApplicationContextMask | EAutomationTestFlags::EngineFilter)
 
-IMPLEMENT_CUSTOM_SIMPLE_AUTOMATION_TEST(FBackChannelTestOSCBundle, FBackChannelTestOSCBase, "BackChannel.TestOSCBundle", EAutomationTestFlags::ApplicationContextMask | EAutomationTestFlags::EngineFilter)
+IMPLEMENT_CUSTOM_SIMPLE_AUTOMATION_TEST(FBackChannelTestOSCBundle, FBackChannelTestOSCBase, "BackChannel.TestOSCBundle", EAutomationTestFlags_ApplicationContextMask | EAutomationTestFlags::EngineFilter)
 
-IMPLEMENT_CUSTOM_SIMPLE_AUTOMATION_TEST(FBackChannelTestOSCBundleWithMessages, FBackChannelTestOSCBase, "BackChannel.TestOSCBundleWithMessages", EAutomationTestFlags::ApplicationContextMask | EAutomationTestFlags::EngineFilter)
+IMPLEMENT_CUSTOM_SIMPLE_AUTOMATION_TEST(FBackChannelTestOSCBundleWithMessages, FBackChannelTestOSCBase, "BackChannel.TestOSCBundleWithMessages", EAutomationTestFlags_ApplicationContextMask | EAutomationTestFlags::EngineFilter)
 
 
 bool FBackChannelTestOSCMessage::RunTest(const FString& Parameters)
diff --git a/Engine/Plugins/Experimental/BackChannel/Source/BackChannel/Private/Transport/TestBackChannelTransport.cpp b/Engine/Plugins/Experimental/BackChannel/Source/BackChannel/Private/Transport/TestBackChannelTransport.cpp
--- a/Engine/Plugins/Experimental/BackChannel/Source/BackChannel/Private/Transport/TestBackChannelTransport.cpp
+++ b/Engine/Plugins/Experimental/BackChannel/Source/BackChannel/Private/Transport/TestBackChannelTransport.cpp
@@ -75,7 +75,7 @@ public:
 
 };
 
-IMPLEMENT_CUSTOM_SIMPLE_AUTOMATION_TEST(FBackChannelTestCreate, FBackChannelTestTransport, "BackChannel.TestTransport", EAutomationTestFlags::ApplicationContextMask | EAutomationTestFlags::EngineFilter)
+IMPLEMENT_CUSTOM_SIMPLE_AUTOMATION_TEST(FBackChannelTestCreate, FBackChannelTestTransport, "BackChannel.TestTransport", EAutomationTestFlags_ApplicationContextMask | EAutomationTestFlags::EngineFilter)
 
 bool FBackChannelTestCreate::RunTest(const FString& Parameters)
 {
@@ -86,7 +86,7 @@ bool FBackChannelTestCreate::RunTest(const FString& Parameters)
 }
 
 
-IMPLEMENT_CUSTOM_SIMPLE_AUTOMATION_TEST(FBackChannelTestConnect, FBackChannelTestTransport, "Project.BackChannel.TestConnect", EAutomationTestFlags::ApplicationContextMask | EAutomationTestFlags::EngineFilter)
+IMPLEMENT_CUSTOM_SIMPLE_AUTOMATION_TEST(FBackChannelTestConnect, FBackChannelTestTransport, "Project.BackChannel.TestConnect", EAutomationTestFlags_ApplicationContextMask | EAutomationTestFlags::EngineFilter)
 
 bool FBackChannelTestConnect::RunTest(const FString& Parameters)
 {
@@ -98,7 +98,7 @@ bool FBackChannelTestConnect::RunTest(const FString& Parameters)
 	return ListenerConnection.IsValid() && ClientConnection.IsValid() && AcceptedConnection.IsValid();	
 }
 
-IMPLEMENT_CUSTOM_SIMPLE_AUTOMATION_TEST(FBackChannelTestSendReceive, FBackChannelTestTransport, "Project.BackChannel.TestConnect", EAutomationTestFlags::ApplicationContextMask | EAutomationTestFlags::EngineFilter)
+IMPLEMENT_CUSTOM_SIMPLE_AUTOMATION_TEST(FBackChannelTestSendReceive, FBackChannelTestTransport, "Project.BackChannel.TestConnect", EAutomationTestFlags_ApplicationContextMask | EAutomationTestFlags::EngineFilter)
 
 bool FBackChannelTestSendReceive::RunTest(const FString& Parameters)
 {
diff --git a/Engine/Plugins/Experimental/JWT/Source/JWT/Private/Tests/JwtTests.cpp b/Engine/Plugins/Experimental/JWT/Source/JWT/Private/Tests/JwtTests.cpp
--- a/Engine/Plugins/Experimental/JWT/Source/JWT/Private/Tests/JwtTests.cpp
+++ b/Engine/Plugins/Experimental/JWT/Source/JWT/Private/Tests/JwtTests.cpp
@@ -162,7 +162,7 @@ namespace TestJwts
  * fail but were expected to fail.
  */
 IMPLEMENT_SIMPLE_AUTOMATION_TEST(FJwtVerificationTest, "JWT.Verification",
-	EAutomationTestFlags::ApplicationContextMask | EAutomationTestFlags::EngineFilter)
+	EAutomationTestFlags_ApplicationContextMask | EAutomationTestFlags::EngineFilter)
 bool FJwtVerificationTest::RunTest(const FString& Parameters)
 {
 	FJwtAlgorithm_RS256 Verifier;
diff --git a/Engine/Plugins/Experimental/NNEDenoiser/Source/NNEDenoiser/Private/NNEDenoiserUnitTestOidn.cpp b/Engine/Plugins/Experimental/NNEDenoiser/Source/NNEDenoiser/Private/NNEDenoiserUnitTestOidn.cpp
--- a/Engine/Plugins/Experimental/NNEDenoiser/Source/NNEDenoiser/Private/NNEDenoiserUnitTestOidn.cpp
+++ b/Engine/Plugins/Experimental/NNEDenoiser/Source/NNEDenoiser/Private/NNEDenoiserUnitTestOidn.cpp
@@ -28,7 +28,7 @@ BEGIN_SHADER_PARAMETER_STRUCT(FTransferFunctionTestDownloadParameters, )
 	RDG_TEXTURE_ACCESS(OutputTexture, ERHIAccess::CopySrc)
 END_SHADER_PARAMETER_STRUCT()
 
-IMPLEMENT_SIMPLE_AUTOMATION_TEST(FTransferFunctionTest, "NNEDenoiser.UnitTests.OIDN.TransferFunction.Forward", EAutomationTestFlags::ApplicationContextMask | EAutomationTestFlags_FeatureMask | EAutomationTestFlags::EngineFilter)
+IMPLEMENT_SIMPLE_AUTOMATION_TEST(FTransferFunctionTest, "NNEDenoiser.UnitTests.OIDN.TransferFunction.Forward", EAutomationTestFlags_ApplicationContextMask | EAutomationTestFlags_FeatureMask | EAutomationTestFlags::EngineFilter)
 bool FTransferFunctionTest::RunTest(const FString& Parameter)
 {
 	const int32 Width = 1920;
diff --git a/Engine/Plugins/Experimental/NNEDenoiser/Source/NNEDenoiser/Private/NNEDenoiserUnitTests.cpp b/Engine/Plugins/Experimental/NNEDenoiser/Source/NNEDenoiser/Private/NNEDenoiserUnitTests.cpp
--- a/Engine/Plugins/Experimental/NNEDenoiser/Source/NNEDenoiser/Private/NNEDenoiserUnitTests.cpp
+++ b/Engine/Plugins/Experimental/NNEDenoiser/Source/NNEDenoiser/Private/NNEDenoiserUnitTests.cpp
@@ -24,7 +24,7 @@ BEGIN_SHADER_PARAMETER_STRUCT(FAutoExposureTestDownloadParameters, )
 	RDG_BUFFER_ACCESS(OutputBuffer, ERHIAccess::CopySrc)
 END_SHADER_PARAMETER_STRUCT()
 
-IMPLEMENT_SIMPLE_AUTOMATION_TEST(FAutoExposureTest, "NNEDenoiser.UnitTests.AutoExposure", EAutomationTestFlags::ApplicationContextMask | EAutomationTestFlags_FeatureMask | EAutomationTestFlags::EngineFilter)
+IMPLEMENT_SIMPLE_AUTOMATION_TEST(FAutoExposureTest, "NNEDenoiser.UnitTests.AutoExposure", EAutomationTestFlags_ApplicationContextMask | EAutomationTestFlags_FeatureMask | EAutomationTestFlags::EngineFilter)
 bool FAutoExposureTest::RunTest(const FString& Parameter)
 {
 	const int32 Width = 1920;
diff --git a/Engine/Plugins/Experimental/QuicMessaging/Source/QuicMessaging/Private/Tests/QuicMessageTransportTest.cpp b/Engine/Plugins/Experimental/QuicMessaging/Source/QuicMessaging/Private/Tests/QuicMessageTransportTest.cpp
--- a/Engine/Plugins/Experimental/QuicMessaging/Source/QuicMessaging/Private/Tests/QuicMessageTransportTest.cpp
+++ b/Engine/Plugins/Experimental/QuicMessaging/Source/QuicMessaging/Private/Tests/QuicMessageTransportTest.cpp
@@ -21,7 +21,7 @@
 
 IMPLEMENT_SIMPLE_AUTOMATION_TEST(FQuicMessageTransportTest,
 	"System.Core.Messaging.Transports.Quic.QuicMessageTransport (may take some minutes!)",
-	EAutomationTestFlags::ApplicationContextMask | EAutomationTestFlags::EngineFilter)
+	EAutomationTestFlags_ApplicationContextMask | EAutomationTestFlags::EngineFilter)
 
 
 /**
diff --git a/Engine/Plugins/Experimental/StormSync/Source/StormSyncDrives/Private/Tests/StormSyncDrivesModule.spec.cpp b/Engine/Plugins/Experimental/StormSync/Source/StormSyncDrives/Private/Tests/StormSyncDrivesModule.spec.cpp
--- a/Engine/Plugins/Experimental/StormSync/Source/StormSyncDrives/Private/Tests/StormSyncDrivesModule.spec.cpp
+++ b/Engine/Plugins/Experimental/StormSync/Source/StormSyncDrives/Private/Tests/StormSyncDrivesModule.spec.cpp
@@ -7,7 +7,7 @@
 #include "Misc/PackageName.h"
 #include "StormSyncDrivesSettings.h"
 
-BEGIN_DEFINE_SPEC(FStormSyncDrivesModuleSpec, "StormSync.StormSyncDrives", EAutomationTestFlags::ProductFilter | EAutomationTestFlags::ApplicationContextMask)
+BEGIN_DEFINE_SPEC(FStormSyncDrivesModuleSpec, "StormSync.StormSyncDrives", EAutomationTestFlags::ProductFilter | EAutomationTestFlags_ApplicationContextMask)
 
 	FString TestMountPoint;
 	FString TestContentPath;
diff --git a/Engine/Plugins/Experimental/StormSync/Source/StormSyncDrives/Private/Tests/StormSyncDrivesUtils.spec.cpp b/Engine/Plugins/Experimental/StormSync/Source/StormSyncDrives/Private/Tests/StormSyncDrivesUtils.spec.cpp
--- a/Engine/Plugins/Experimental/StormSync/Source/StormSyncDrives/Private/Tests/StormSyncDrivesUtils.spec.cpp
+++ b/Engine/Plugins/Experimental/StormSync/Source/StormSyncDrives/Private/Tests/StormSyncDrivesUtils.spec.cpp
@@ -6,7 +6,7 @@
 #include "StormSyncDrivesSettings.h"
 #include "StormSyncDrivesUtils.h"
 
-BEGIN_DEFINE_SPEC(FStormSyncDrivesUtilsSpec, "StormSync.StormSyncDrives", EAutomationTestFlags::ProductFilter | EAutomationTestFlags::ApplicationContextMask)
+BEGIN_DEFINE_SPEC(FStormSyncDrivesUtilsSpec, "StormSync.StormSyncDrives", EAutomationTestFlags::ProductFilter | EAutomationTestFlags_ApplicationContextMask)
 
 	/** Helper to assert a duplicate error log */
 	void AssertErrorLog(const int32 InErrorIndex, const int32 InTestCaseIndex, const FStormSyncMountPointConfig& InTestCase, const TArray<FText>& InValidationErrors)
diff --git a/Engine/Plugins/Experimental/StormSync/Source/StormSyncEditor/Private/Tests/StormSyncArchiveFactory.spec.cpp b/Engine/Plugins/Experimental/StormSync/Source/StormSyncEditor/Private/Tests/StormSyncArchiveFactory.spec.cpp
--- a/Engine/Plugins/Experimental/StormSync/Source/StormSyncEditor/Private/Tests/StormSyncArchiveFactory.spec.cpp
+++ b/Engine/Plugins/Experimental/StormSync/Source/StormSyncEditor/Private/Tests/StormSyncArchiveFactory.spec.cpp
@@ -3,7 +3,7 @@
 #include "Factories/StormSyncArchiveFactory.h"
 #include "Misc/AutomationTest.h"
 
-BEGIN_DEFINE_SPEC(FStormSyncArchiveFactorySpec, "StormSync.StormSyncEditor.StormSyncArchiveFactory", EAutomationTestFlags::ProductFilter | EAutomationTestFlags::ApplicationContextMask)
+BEGIN_DEFINE_SPEC(FStormSyncArchiveFactorySpec, "StormSync.StormSyncEditor.StormSyncArchiveFactory", EAutomationTestFlags::ProductFilter | EAutomationTestFlags_ApplicationContextMask)
 
 	UStormSyncArchiveFactory* Factory;
 
diff --git a/Engine/Plugins/Experimental/StormSync/Source/StormSyncEditor/Private/Tests/StormSyncExportWizard.spec.cpp b/Engine/Plugins/Experimental/StormSync/Source/StormSyncEditor/Private/Tests/StormSyncExportWizard.spec.cpp
--- a/Engine/Plugins/Experimental/StormSync/Source/StormSyncEditor/Private/Tests/StormSyncExportWizard.spec.cpp
+++ b/Engine/Plugins/Experimental/StormSync/Source/StormSyncEditor/Private/Tests/StormSyncExportWizard.spec.cpp
@@ -5,7 +5,7 @@
 #include "StormSyncCoreSettings.h"
 #include "Slate/SStormSyncExportWizard.h"
 
-BEGIN_DEFINE_SPEC(FStormSyncExportWizardSpec, "StormSync.StormSyncEditor.SStormSyncExportWizard", EAutomationTestFlags::ProductFilter | EAutomationTestFlags::ApplicationContextMask)
+BEGIN_DEFINE_SPEC(FStormSyncExportWizardSpec, "StormSync.StormSyncEditor.SStormSyncExportWizard", EAutomationTestFlags::ProductFilter | EAutomationTestFlags_ApplicationContextMask)
 
 
 END_DEFINE_SPEC(FStormSyncExportWizardSpec)
diff --git a/Engine/Plugins/Experimental/StormSync/Source/StormSyncTests/Private/StormSyncCore.spec.cpp b/Engine/Plugins/Experimental/StormSync/Source/StormSyncTests/Private/StormSyncCore.spec.cpp
--- a/Engine/Plugins/Experimental/StormSync/Source/StormSyncTests/Private/StormSyncCore.spec.cpp
+++ b/Engine/Plugins/Experimental/StormSync/Source/StormSyncTests/Private/StormSyncCore.spec.cpp
@@ -9,7 +9,7 @@
 
 #pragma warning(disable : 6011)
 
-BEGIN_DEFINE_SPEC(FStormSyncCoreSpec, "StormSync.StormSyncCore", EAutomationTestFlags::ProductFilter | EAutomationTestFlags::ApplicationContextMask)
+BEGIN_DEFINE_SPEC(FStormSyncCoreSpec, "StormSync.StormSyncCore", EAutomationTestFlags::ProductFilter | EAutomationTestFlags_ApplicationContextMask)
 
 	FString InvalidWidgetFixturePath = TEXT("/StormSync/Fixtures/Invalid");
 
diff --git a/Engine/Plugins/Experimental/StormSync/Source/StormSyncTests/Private/StormSyncPackageDescriptor.spec.cpp b/Engine/Plugins/Experimental/StormSync/Source/StormSyncTests/Private/StormSyncPackageDescriptor.spec.cpp
--- a/Engine/Plugins/Experimental/StormSync/Source/StormSyncTests/Private/StormSyncPackageDescriptor.spec.cpp
+++ b/Engine/Plugins/Experimental/StormSync/Source/StormSyncTests/Private/StormSyncPackageDescriptor.spec.cpp
@@ -9,7 +9,7 @@
 #include "StormSyncCoreUtils.h"
 #include "StormSyncDrivesSettings.h"
 
-BEGIN_DEFINE_SPEC(FStormSyncPackageDescriptorSpec, "StormSync.StormSyncCore", EAutomationTestFlags::ProductFilter | EAutomationTestFlags::ApplicationContextMask)
+BEGIN_DEFINE_SPEC(FStormSyncPackageDescriptorSpec, "StormSync.StormSyncCore", EAutomationTestFlags::ProductFilter | EAutomationTestFlags_ApplicationContextMask)
 
 	FString TestMountPoint = TEXT("/StormSyncTests_Temp");
 	FString TestMountDirectory = FPaths::ProjectSavedDir() / TEXT("StormSyncTests");
diff --git a/Engine/Plugins/Experimental/StormSync/Source/StormSyncTransportClient/Private/Tests/StormSyncTransportClientEndpoint.spec.cpp b/Engine/Plugins/Experimental/StormSync/Source/StormSyncTransportClient/Private/Tests/StormSyncTransportClientEndpoint.spec.cpp
--- a/Engine/Plugins/Experimental/StormSync/Source/StormSyncTransportClient/Private/Tests/StormSyncTransportClientEndpoint.spec.cpp
+++ b/Engine/Plugins/Experimental/StormSync/Source/StormSyncTransportClient/Private/Tests/StormSyncTransportClientEndpoint.spec.cpp
@@ -242,7 +242,7 @@ private:
 	}
 };
 
-BEGIN_DEFINE_SPEC(FStormSyncTransportClientEndpointSpec, "StormSync.StormSyncTransportClient.StormSyncTransportClientEndpoint", EAutomationTestFlags::ProductFilter | EAutomationTestFlags::ApplicationContextMask)
+BEGIN_DEFINE_SPEC(FStormSyncTransportClientEndpointSpec, "StormSync.StormSyncTransportClient.StormSyncTransportClientEndpoint", EAutomationTestFlags::ProductFilter | EAutomationTestFlags_ApplicationContextMask)
 
 	static constexpr const TCHAR* EndpointFriendlyName = TEXT("Test Suite ClientEndpoint (Message System)");
 	TSharedPtr<FStormSyncTransportClientEndpoint> ClientEndpoint;
diff --git a/Engine/Plugins/Experimental/StormSync/Source/StormSyncTransportCore/Private/Tests/StormSyncTransportCommandUtils.spec.cpp b/Engine/Plugins/Experimental/StormSync/Source/StormSyncTransportCore/Private/Tests/StormSyncTransportCommandUtils.spec.cpp
--- a/Engine/Plugins/Experimental/StormSync/Source/StormSyncTransportCore/Private/Tests/StormSyncTransportCommandUtils.spec.cpp
+++ b/Engine/Plugins/Experimental/StormSync/Source/StormSyncTransportCore/Private/Tests/StormSyncTransportCommandUtils.spec.cpp
@@ -5,7 +5,7 @@
 #include "StormSyncTransportSettings.h"
 #include "Utils/StormSyncTransportCommandUtils.h"
 
-BEGIN_DEFINE_SPEC(FStormSyncTransportCommandUtilsSpec, "StormSync.StormSyncTransportCore.StormSyncTransportCommandUtils", EAutomationTestFlags::ProductFilter | EAutomationTestFlags::ApplicationContextMask)
+BEGIN_DEFINE_SPEC(FStormSyncTransportCommandUtilsSpec, "StormSync.StormSyncTransportCore.StormSyncTransportCommandUtils", EAutomationTestFlags::ProductFilter | EAutomationTestFlags_ApplicationContextMask)
 	FString OriginalCmdLine;
 END_DEFINE_SPEC(FStormSyncTransportCommandUtilsSpec)
 
diff --git a/Engine/Plugins/Experimental/StormSync/Source/StormSyncTransportCore/Private/Tests/StormSyncTransportMessages.spec.cpp b/Engine/Plugins/Experimental/StormSync/Source/StormSyncTransportCore/Private/Tests/StormSyncTransportMessages.spec.cpp
--- a/Engine/Plugins/Experimental/StormSync/Source/StormSyncTransportCore/Private/Tests/StormSyncTransportMessages.spec.cpp
+++ b/Engine/Plugins/Experimental/StormSync/Source/StormSyncTransportCore/Private/Tests/StormSyncTransportMessages.spec.cpp
@@ -6,7 +6,7 @@
 #include "StormSyncTransportMessages.h"
 #include "Engine/Engine.h"
 
-BEGIN_DEFINE_SPEC(FStormSyncTransportMessagesSpec, "StormSync.StormSyncTransportCore.StormSyncTransportMessages", EAutomationTestFlags::ProductFilter | EAutomationTestFlags::ApplicationContextMask)
+BEGIN_DEFINE_SPEC(FStormSyncTransportMessagesSpec, "StormSync.StormSyncTransportCore.StormSyncTransportMessages", EAutomationTestFlags::ProductFilter | EAutomationTestFlags_ApplicationContextMask)
 
 END_DEFINE_SPEC(FStormSyncTransportMessagesSpec)
 
diff --git a/Engine/Plugins/Experimental/StormSync/Source/StormSyncTransportCore/Private/Tests/StormSyncTransportNetworkUtils.spec.cpp b/Engine/Plugins/Experimental/StormSync/Source/StormSyncTransportCore/Private/Tests/StormSyncTransportNetworkUtils.spec.cpp
--- a/Engine/Plugins/Experimental/StormSync/Source/StormSyncTransportCore/Private/Tests/StormSyncTransportNetworkUtils.spec.cpp
+++ b/Engine/Plugins/Experimental/StormSync/Source/StormSyncTransportCore/Private/Tests/StormSyncTransportNetworkUtils.spec.cpp
@@ -5,7 +5,7 @@
 #include "StormSyncTransportNetworkUtils.h"
 #include "StormSyncTransportSettings.h"
 
-BEGIN_DEFINE_SPEC(FStormSyncTransportNetworkUtilsSpec, "StormSync.StormSyncTransportCore.StormSyncTransportNetworkUtils", EAutomationTestFlags::ProductFilter | EAutomationTestFlags::ApplicationContextMask)
+BEGIN_DEFINE_SPEC(FStormSyncTransportNetworkUtilsSpec, "StormSync.StormSyncTransportCore.StormSyncTransportNetworkUtils", EAutomationTestFlags::ProductFilter | EAutomationTestFlags_ApplicationContextMask)
 
 END_DEFINE_SPEC(FStormSyncTransportNetworkUtilsSpec)
 
diff --git a/Engine/Plugins/Experimental/StormSync/Source/StormSyncTransportServer/Private/Tests/StormSyncServiceDiscovery.spec.cpp b/Engine/Plugins/Experimental/StormSync/Source/StormSyncTransportServer/Private/Tests/StormSyncServiceDiscovery.spec.cpp
--- a/Engine/Plugins/Experimental/StormSync/Source/StormSyncTransportServer/Private/Tests/StormSyncServiceDiscovery.spec.cpp
+++ b/Engine/Plugins/Experimental/StormSync/Source/StormSyncTransportServer/Private/Tests/StormSyncServiceDiscovery.spec.cpp
@@ -54,7 +54,7 @@ private:
 	}
 };
 
-BEGIN_DEFINE_SPEC(FStormSyncServiceDiscoverySpec, "StormSync.StormSyncTransportServer.StormSyncServiceDiscovery", EAutomationTestFlags::ProductFilter | EAutomationTestFlags::ApplicationContextMask)
+BEGIN_DEFINE_SPEC(FStormSyncServiceDiscoverySpec, "StormSync.StormSyncTransportServer.StormSyncServiceDiscovery", EAutomationTestFlags::ProductFilter | EAutomationTestFlags_ApplicationContextMask)
 
 	/** Mock version of StormSyncDiscoveryManager (used in "StormSyncHeartbeatEmitter" suite) */
 	TUniquePtr<FMockDiscoveryManager> MockDiscoveryManager;
diff --git a/Engine/Plugins/Experimental/StormSync/Source/StormSyncTransportServer/Private/Tests/StormSyncTransportServerEndpoint.spec.cpp b/Engine/Plugins/Experimental/StormSync/Source/StormSyncTransportServer/Private/Tests/StormSyncTransportServerEndpoint.spec.cpp
--- a/Engine/Plugins/Experimental/StormSync/Source/StormSyncTransportServer/Private/Tests/StormSyncTransportServerEndpoint.spec.cpp
+++ b/Engine/Plugins/Experimental/StormSync/Source/StormSyncTransportServer/Private/Tests/StormSyncTransportServerEndpoint.spec.cpp
@@ -12,7 +12,7 @@
 #include "StormSyncTransportServerEndpoint.h"
 #include "StormSyncTransportSettings.h"
 
-BEGIN_DEFINE_SPEC(FStormSyncTransportServerEndpointSpec, "StormSync.StormSyncTransportServer.StormSyncTransportServerEndpoint", EAutomationTestFlags::ProductFilter | EAutomationTestFlags::ApplicationContextMask)
+BEGIN_DEFINE_SPEC(FStormSyncTransportServerEndpointSpec, "StormSync.StormSyncTransportServer.StormSyncTransportServerEndpoint", EAutomationTestFlags::ProductFilter | EAutomationTestFlags_ApplicationContextMask)
 
 	static constexpr const TCHAR* EndpointFriendlyName = TEXT("Test Suite ServerEndpoint (Message System)");
 
diff --git a/Engine/Plugins/Experimental/StormSync/Source/StormSyncTransportServer/Private/Tests/StormSyncTransportTcpServer.spec.cpp b/Engine/Plugins/Experimental/StormSync/Source/StormSyncTransportServer/Private/Tests/StormSyncTransportTcpServer.spec.cpp
--- a/Engine/Plugins/Experimental/StormSync/Source/StormSyncTransportServer/Private/Tests/StormSyncTransportTcpServer.spec.cpp
+++ b/Engine/Plugins/Experimental/StormSync/Source/StormSyncTransportServer/Private/Tests/StormSyncTransportTcpServer.spec.cpp
@@ -6,7 +6,7 @@
 #include "StormSyncTransportSettings.h"
 
 
-BEGIN_DEFINE_SPEC(FStormSyncTransportTcpServerSpec, "StormSync.StormSyncTransportServer.StormSyncTransportTcpServer", EAutomationTestFlags::ProductFilter | EAutomationTestFlags::ApplicationContextMask)
+BEGIN_DEFINE_SPEC(FStormSyncTransportTcpServerSpec, "StormSync.StormSyncTransportServer.StormSyncTransportTcpServer", EAutomationTestFlags::ProductFilter | EAutomationTestFlags_ApplicationContextMask)
 
 	/** Our instance of tcp listener wrapper */
 	TUniquePtr<FStormSyncTransportTcpServer> MockTcpServer;
diff --git a/Engine/Plugins/Experimental/Web/WebAPI/Source/WebAPIEditor/Private/Tests/WebAPIEditorUtilities.spec.cpp b/Engine/Plugins/Experimental/Web/WebAPI/Source/WebAPIEditor/Private/Tests/WebAPIEditorUtilities.spec.cpp
--- a/Engine/Plugins/Experimental/Web/WebAPI/Source/WebAPIEditor/Private/Tests/WebAPIEditorUtilities.spec.cpp
+++ b/Engine/Plugins/Experimental/Web/WebAPI/Source/WebAPIEditor/Private/Tests/WebAPIEditorUtilities.spec.cpp
@@ -8,7 +8,7 @@
 
 BEGIN_DEFINE_SPEC(FWebAPIEditorUtilitiesSpec,
 	TEXT("Plugins.WebAPI.Editor.Utilities"),
-	EAutomationTestFlags::EditorContext | EAutomationTestFlags::EngineFilter | EAutomationTestFlags::ApplicationContextMask)
+	EAutomationTestFlags::EditorContext | EAutomationTestFlags::EngineFilter | EAutomationTestFlags_ApplicationContextMask)
 
 END_DEFINE_SPEC(FWebAPIEditorUtilitiesSpec)
 
diff --git a/Engine/Plugins/Experimental/Web/WebAPI/Source/WebAPIEditor/Private/Tests/WebAPIJsonUtilities.spec.cpp b/Engine/Plugins/Experimental/Web/WebAPI/Source/WebAPIEditor/Private/Tests/WebAPIJsonUtilities.spec.cpp
--- a/Engine/Plugins/Experimental/Web/WebAPI/Source/WebAPIEditor/Private/Tests/WebAPIJsonUtilities.spec.cpp
+++ b/Engine/Plugins/Experimental/Web/WebAPI/Source/WebAPIEditor/Private/Tests/WebAPIJsonUtilities.spec.cpp
@@ -9,7 +9,7 @@
 
 BEGIN_DEFINE_SPEC(FWebAPIEditorJsonSpec,
 	TEXT("Plugins.WebAPI.Editor.Json"),
-	EAutomationTestFlags::EditorContext | EAutomationTestFlags::EngineFilter | EAutomationTestFlags::ApplicationContextMask)
+	EAutomationTestFlags::EditorContext | EAutomationTestFlags::EngineFilter | EAutomationTestFlags_ApplicationContextMask)
 
 	TSharedPtr<FJsonObject> JsonObject;
 
diff --git a/Engine/Plugins/Experimental/Web/WebAPI/Source/WebAPIOpenAPI/Private/Tests/WebAPIOpenAPI3.spec.cpp b/Engine/Plugins/Experimental/Web/WebAPI/Source/WebAPIOpenAPI/Private/Tests/WebAPIOpenAPI3.spec.cpp
--- a/Engine/Plugins/Experimental/Web/WebAPI/Source/WebAPIOpenAPI/Private/Tests/WebAPIOpenAPI3.spec.cpp
+++ b/Engine/Plugins/Experimental/Web/WebAPI/Source/WebAPIOpenAPI/Private/Tests/WebAPIOpenAPI3.spec.cpp
@@ -19,7 +19,7 @@
 
 BEGIN_DEFINE_SPEC(FWebAPIOpenAPI3Spec,
 				"Plugins.WebAPI.OpenAPI3",
-				EAutomationTestFlags::EditorContext | EAutomationTestFlags::EngineFilter | EAutomationTestFlags::ApplicationContextMask)
+				EAutomationTestFlags::EditorContext | EAutomationTestFlags::EngineFilter | EAutomationTestFlags_ApplicationContextMask)
 
 	TSharedPtr<UE::WebAPI::OpenAPI::V3::FOpenAPIObject> InputDefinition;
 	TStrongObjectPtr<UWebAPIDefinition> OutputDefinition;
diff --git a/Engine/Plugins/Experimental/Web/WebAPI/Source/WebAPIOpenAPI/Private/Tests/WebAPISwagger.spec.cpp b/Engine/Plugins/Experimental/Web/WebAPI/Source/WebAPIOpenAPI/Private/Tests/WebAPISwagger.spec.cpp
--- a/Engine/Plugins/Experimental/Web/WebAPI/Source/WebAPIOpenAPI/Private/Tests/WebAPISwagger.spec.cpp
+++ b/Engine/Plugins/Experimental/Web/WebAPI/Source/WebAPIOpenAPI/Private/Tests/WebAPISwagger.spec.cpp
@@ -15,7 +15,7 @@
 
 BEGIN_DEFINE_SPEC(FWebAPISwaggerSpec,
 	TEXT("Plugins.WebAPI.Swagger"),
-	EAutomationTestFlags::EditorContext | EAutomationTestFlags::EngineFilter | EAutomationTestFlags::ApplicationContextMask)
+	EAutomationTestFlags::EditorContext | EAutomationTestFlags::EngineFilter | EAutomationTestFlags_ApplicationContextMask)
 
 	TSharedPtr<UE::WebAPI::OpenAPI::V2::FSwagger> InputDefinition;
 	TStrongObjectPtr<UWebAPIDefinition> OutputDefinition;
diff --git a/Engine/Plugins/Media/ElectraPlayer/Source/ElectraPlayerRuntime/Private/Tests/Parser.spec.cpp b/Engine/Plugins/Media/ElectraPlayer/Source/ElectraPlayerRuntime/Private/Tests/Parser.spec.cpp
--- a/Engine/Plugins/Media/ElectraPlayer/Source/ElectraPlayerRuntime/Private/Tests/Parser.spec.cpp
+++ b/Engine/Plugins/Media/ElectraPlayer/Source/ElectraPlayerRuntime/Private/Tests/Parser.spec.cpp
@@ -7,7 +7,7 @@
 #include "Player/HLS/LHLSTags.h"
 #include "CoreMinimal.h"
 
-BEGIN_DEFINE_SPEC(FParserSpec, "ElectraPlayer.Unit", EAutomationTestFlags::ProductFilter | EAutomationTestFlags::ApplicationContextMask)
+BEGIN_DEFINE_SPEC(FParserSpec, "ElectraPlayer.Unit", EAutomationTestFlags::ProductFilter | EAutomationTestFlags_ApplicationContextMask)
 END_DEFINE_SPEC(FParserSpec)
 
 void FParserSpec::Define()
diff --git a/Engine/Plugins/Messaging/UdpMessaging/Source/UdpMessaging/Private/Tests/UdpMessageSegmenterTest.cpp b/Engine/Plugins/Messaging/UdpMessaging/Source/UdpMessaging/Private/Tests/UdpMessageSegmenterTest.cpp
--- a/Engine/Plugins/Messaging/UdpMessaging/Source/UdpMessaging/Private/Tests/UdpMessageSegmenterTest.cpp
+++ b/Engine/Plugins/Messaging/UdpMessaging/Source/UdpMessaging/Private/Tests/UdpMessageSegmenterTest.cpp
@@ -9,7 +9,7 @@
 #include "Transport/UdpMessageSegmenter.h"
 
 
-IMPLEMENT_SIMPLE_AUTOMATION_TEST(FUdpMessageSegmenterTest, "System.Core.Messaging.Transports.Udp.UdpMessageSegmenter", EAutomationTestFlags::ApplicationContextMask | EAutomationTestFlags::EngineFilter)
+IMPLEMENT_SIMPLE_AUTOMATION_TEST(FUdpMessageSegmenterTest, "System.Core.Messaging.Transports.Udp.UdpMessageSegmenter", EAutomationTestFlags_ApplicationContextMask | EAutomationTestFlags::EngineFilter)
 
 
 void RunSegmentationTest(FAutomationTestBase& Test, uint32 MessageSize, uint16 SegmentSize, bool WithReliableMessages)
diff --git a/Engine/Plugins/Messaging/UdpMessaging/Source/UdpMessaging/Private/Tests/UdpMessageTransportTest.cpp b/Engine/Plugins/Messaging/UdpMessaging/Source/UdpMessaging/Private/Tests/UdpMessageTransportTest.cpp
--- a/Engine/Plugins/Messaging/UdpMessaging/Source/UdpMessaging/Private/Tests/UdpMessageTransportTest.cpp
+++ b/Engine/Plugins/Messaging/UdpMessaging/Source/UdpMessaging/Private/Tests/UdpMessageTransportTest.cpp
@@ -9,7 +9,7 @@
 #include "HAL/CriticalSection.h"
 #include "Misc/ScopeLock.h"
 
-IMPLEMENT_SIMPLE_AUTOMATION_TEST(FUdpMessageTransportTest, "System.Core.Messaging.Transports.Udp.UdpMessageTransport", EAutomationTestFlags::ApplicationContextMask | EAutomationTestFlags::EngineFilter)
+IMPLEMENT_SIMPLE_AUTOMATION_TEST(FUdpMessageTransportTest, "System.Core.Messaging.Transports.Udp.UdpMessageTransport", EAutomationTestFlags_ApplicationContextMask | EAutomationTestFlags::EngineFilter)
 
 class FUdpMessageTransportTestState
 	: public IMessageTransportHandler
@@ -97,7 +97,7 @@ private:
 	TSharedPtr<IMessageContext, ESPMode::ThreadSafe> LastMessageContext;
 };
 
-IMPLEMENT_SIMPLE_AUTOMATION_TEST(FUdpMessageTransportTestExclusionList, "System.Core.Messaging.Transports.Udp.UdpMessageTransportExclusionList", EAutomationTestFlags::ApplicationContextMask | EAutomationTestFlags::EngineFilter)
+IMPLEMENT_SIMPLE_AUTOMATION_TEST(FUdpMessageTransportTestExclusionList, "System.Core.Messaging.Transports.Udp.UdpMessageTransportExclusionList", EAutomationTestFlags_ApplicationContextMask | EAutomationTestFlags::EngineFilter)
 
 bool FUdpMessageTransportTestExclusionList::RunTest(const FString& Parameters)
 {
diff --git a/Engine/Plugins/Messaging/UdpMessaging/Source/UdpMessaging/Private/Tests/UdpSerializeMessageTaskTest.cpp b/Engine/Plugins/Messaging/UdpMessaging/Source/UdpMessaging/Private/Tests/UdpSerializeMessageTaskTest.cpp
--- a/Engine/Plugins/Messaging/UdpMessaging/Source/UdpMessaging/Private/Tests/UdpSerializeMessageTaskTest.cpp
+++ b/Engine/Plugins/Messaging/UdpMessaging/Source/UdpMessaging/Private/Tests/UdpSerializeMessageTaskTest.cpp
@@ -10,7 +10,7 @@
 #include "Tests/UdpMessagingTestTypes.h"
 
 
-IMPLEMENT_SIMPLE_AUTOMATION_TEST(FUdpSerializeMessageTaskTest, "System.Core.Messaging.Transports.Udp.UdpSerializedMessage", EAutomationTestFlags::ApplicationContextMask | EAutomationTestFlags::EngineFilter)
+IMPLEMENT_SIMPLE_AUTOMATION_TEST(FUdpSerializeMessageTaskTest, "System.Core.Messaging.Transports.Udp.UdpSerializedMessage", EAutomationTestFlags_ApplicationContextMask | EAutomationTestFlags::EngineFilter)
 
 
 namespace UdpSerializeMessageTaskTest
diff --git a/Engine/Plugins/NetcodeUnitTest/NetcodeUnitTest/Source/NetcodeUnitTest/Private/NUTUtilReflectionParser.cpp b/Engine/Plugins/NetcodeUnitTest/NetcodeUnitTest/Source/NetcodeUnitTest/Private/NUTUtilReflectionParser.cpp
--- a/Engine/Plugins/NetcodeUnitTest/NetcodeUnitTest/Source/NetcodeUnitTest/Private/NUTUtilReflectionParser.cpp
+++ b/Engine/Plugins/NetcodeUnitTest/NetcodeUnitTest/Source/NetcodeUnitTest/Private/NUTUtilReflectionParser.cpp
@@ -728,7 +728,7 @@ bool TestExpression(FAutomationTestBase* Test, const TCHAR* Expression, const TC
 }
 
 IMPLEMENT_SIMPLE_AUTOMATION_TEST(FLTRAssociativtyExpressionsTest, "System.Core.Expression Parser.LTR Operator Associativity",
-									EAutomationTestFlags::ApplicationContextMask | EAutomationTestFlags::SmokeFilter);
+									EAutomationTestFlags_ApplicationContextMask | EAutomationTestFlags::SmokeFilter);
 
 bool FLTRAssociativtyExpressionsTest::RunTest(const FString& Parameters)
 {
@@ -739,7 +739,7 @@ bool FLTRAssociativtyExpressionsTest::RunTest(const FString& Parameters)
 }
 
 IMPLEMENT_SIMPLE_AUTOMATION_TEST(FRTLAssociativtyExpressionsTest, "System.Core.Expression Parser.RTL Operator Associativity",
-									EAutomationTestFlags::ApplicationContextMask | EAutomationTestFlags::SmokeFilter);
+									EAutomationTestFlags_ApplicationContextMask | EAutomationTestFlags::SmokeFilter);
 
 bool FRTLAssociativtyExpressionsTest::RunTest(const FString& Parameters)
 {
diff --git a/Engine/Plugins/Online/OnlineSubsystem/Source/Test/OnlineAchievements.spec.cpp b/Engine/Plugins/Online/OnlineSubsystem/Source/Test/OnlineAchievements.spec.cpp
--- a/Engine/Plugins/Online/OnlineSubsystem/Source/Test/OnlineAchievements.spec.cpp
+++ b/Engine/Plugins/Online/OnlineSubsystem/Source/Test/OnlineAchievements.spec.cpp
@@ -5,7 +5,7 @@
 #include "Utils/OnlineErrors.data.h"
 #include "Utils/OnlineTestCommon.h"
 
-BEGIN_DEFINE_SPEC(FOnlineAchievementsSpec, "OnlineAchievementsInterface", EAutomationTestFlags::EngineFilter | EAutomationTestFlags::ApplicationContextMask)
+BEGIN_DEFINE_SPEC(FOnlineAchievementsSpec, "OnlineAchievementsInterface", EAutomationTestFlags::EngineFilter | EAutomationTestFlags_ApplicationContextMask)
 
 IOnlineIdentityPtr OnlineIdentity;
 IOnlineAchievementsPtr OnlineAchievements;
diff --git a/Engine/Plugins/Online/OnlineSubsystem/Source/Test/OnlineChat.spec.cpp b/Engine/Plugins/Online/OnlineSubsystem/Source/Test/OnlineChat.spec.cpp
--- a/Engine/Plugins/Online/OnlineSubsystem/Source/Test/OnlineChat.spec.cpp
+++ b/Engine/Plugins/Online/OnlineSubsystem/Source/Test/OnlineChat.spec.cpp
@@ -4,7 +4,7 @@
 #include "Online.h"
 #include "Utils/OnlineTestCommon.h"
 
-BEGIN_DEFINE_SPEC(FOnlineChatSpec, "OnlineChatInterface", EAutomationTestFlags::EngineFilter | EAutomationTestFlags::ApplicationContextMask)
+BEGIN_DEFINE_SPEC(FOnlineChatSpec, "OnlineChatInterface", EAutomationTestFlags::EngineFilter | EAutomationTestFlags_ApplicationContextMask)
 
 IOnlineSubsystem* OnlineSubsystem;
 
diff --git a/Engine/Plugins/Online/OnlineSubsystem/Source/Test/OnlineEntitlements.spec.cpp b/Engine/Plugins/Online/OnlineSubsystem/Source/Test/OnlineEntitlements.spec.cpp
--- a/Engine/Plugins/Online/OnlineSubsystem/Source/Test/OnlineEntitlements.spec.cpp
+++ b/Engine/Plugins/Online/OnlineSubsystem/Source/Test/OnlineEntitlements.spec.cpp
@@ -3,7 +3,7 @@
 #include "Online.h"
 #include "Utils/OnlineTestCommon.h"
 
-BEGIN_DEFINE_SPEC(FOnlineEntitlementsSpec, "OnlineEntitlementsInterface", EAutomationTestFlags::EngineFilter | EAutomationTestFlags::ApplicationContextMask)
+BEGIN_DEFINE_SPEC(FOnlineEntitlementsSpec, "OnlineEntitlementsInterface", EAutomationTestFlags::EngineFilter | EAutomationTestFlags_ApplicationContextMask)
 
 IOnlineSubsystem* OnlineSubsystem;
 
diff --git a/Engine/Plugins/Online/OnlineSubsystem/Source/Test/OnlineFriends.spec.cpp b/Engine/Plugins/Online/OnlineSubsystem/Source/Test/OnlineFriends.spec.cpp
--- a/Engine/Plugins/Online/OnlineSubsystem/Source/Test/OnlineFriends.spec.cpp
+++ b/Engine/Plugins/Online/OnlineSubsystem/Source/Test/OnlineFriends.spec.cpp
@@ -6,7 +6,7 @@
 #include "Misc/CommandLine.h"
 #include "Utils/OnlineTestCommon.h"
 
-BEGIN_DEFINE_SPEC(FOnlineFriendsSpec, "OnlineFriendsInterface", EAutomationTestFlags::EngineFilter | EAutomationTestFlags::ApplicationContextMask)
+BEGIN_DEFINE_SPEC(FOnlineFriendsSpec, "OnlineFriendsInterface", EAutomationTestFlags::EngineFilter | EAutomationTestFlags_ApplicationContextMask)
 
 FOnlineTestCommon CommonUtils;
 
diff --git a/Engine/Plugins/Online/OnlineSubsystem/Source/Test/OnlineIdentity.spec.cpp b/Engine/Plugins/Online/OnlineSubsystem/Source/Test/OnlineIdentity.spec.cpp
--- a/Engine/Plugins/Online/OnlineSubsystem/Source/Test/OnlineIdentity.spec.cpp
+++ b/Engine/Plugins/Online/OnlineSubsystem/Source/Test/OnlineIdentity.spec.cpp
@@ -12,7 +12,7 @@ THIRD_PARTY_INCLUDES_START
 THIRD_PARTY_INCLUDES_END
 */
 
-BEGIN_DEFINE_SPEC(FOnlineIdentitySpec, "OnlineIdentityInterface", EAutomationTestFlags::EngineFilter | EAutomationTestFlags::ApplicationContextMask)
+BEGIN_DEFINE_SPEC(FOnlineIdentitySpec, "OnlineIdentityInterface", EAutomationTestFlags::EngineFilter | EAutomationTestFlags_ApplicationContextMask)
 
 FOnlineTestCommon CommonUtils;
 
diff --git a/Engine/Plugins/Online/OnlineSubsystem/Source/Test/OnlineLeaderboard.spec.cpp b/Engine/Plugins/Online/OnlineSubsystem/Source/Test/OnlineLeaderboard.spec.cpp
--- a/Engine/Plugins/Online/OnlineSubsystem/Source/Test/OnlineLeaderboard.spec.cpp
+++ b/Engine/Plugins/Online/OnlineSubsystem/Source/Test/OnlineLeaderboard.spec.cpp
@@ -4,7 +4,7 @@
 #include "Interfaces/OnlineLeaderboardInterface.h"
 #include "Utils/OnlineTestCommon.h"
 
-BEGIN_DEFINE_SPEC(FOnlineLeaderboardSpec, "OnlineLeaderboardInterface", EAutomationTestFlags::EngineFilter | EAutomationTestFlags::ApplicationContextMask)
+BEGIN_DEFINE_SPEC(FOnlineLeaderboardSpec, "OnlineLeaderboardInterface", EAutomationTestFlags::EngineFilter | EAutomationTestFlags_ApplicationContextMask)
 
 IOnlineSubsystem* OnlineSubsystem;
 
