IPv6 support for RemoteSession client and server
BackChannel plugin now listens on both IPv4 and IPv6 if available and can connect to either
NetworkServiceDiscovery filters returned addresses by what's available on the client for iOS and Android

#virtualized
#Synced-CL 52622993

diff --git a/Engine/Plugins/Experimental/BackChannel/Source/BackChannel/Private/Transport/BackChannelConnection.cpp b/Engine/Plugins/Experimental/BackChannel/Source/BackChannel/Private/Transport/BackChannelConnection.cpp
--- a/Engine/Plugins/Experimental/BackChannel/Source/BackChannel/Private/Transport/BackChannelConnection.cpp
+++ b/Engine/Plugins/Experimental/BackChannel/Source/BackChannel/Private/Transport/BackChannelConnection.cpp
@@ -4,6 +4,8 @@
 #include "BackChannelCommon.h"
 #include "HAL/IConsoleManager.h"
 #include "Interfaces/IPv4/IPv4Endpoint.h"
+#include "IPAddress.h"
+#include "SocketSubsystem.h"
 #include "Misc/ConfigCacheIni.h"
 #include "Misc/OutputDeviceRedirector.h"
 #include "Sockets.h"
@@ -119,8 +121,60 @@ bool FBackChannelConnection::Connect(const TCHAR* InEndPoint)
 	IsAttemptingConnection = true;
 
 	FString LocalEndPoint = InEndPoint;
+	ISocketSubsystem* SocketSubsystem = ISocketSubsystem::Get(PLATFORM_SOCKETSUBSYSTEM);
+
+	// Parse the endpoint — try IPv4 first (fast path), then fall back to generic parsing for IPv6
+	TSharedPtr<FInternetAddr> TargetAddr;
+	FIPv4Endpoint EndPointv4;
+	if (FIPv4Endpoint::Parse(LocalEndPoint, EndPointv4))
+	{
+		TargetAddr = EndPointv4.ToInternetAddr();
+	}
+	else
+	{
+		// Try generic parsing for IPv6, e.g. "[::1]:2049"
+		FString AddressPart;
+		int32 Port = 0;
+		if (LocalEndPoint.StartsWith(TEXT("[")))
+		{
+			// Format: [IPv6Address]:Port
+			int32 CloseBracket;
+			if (LocalEndPoint.FindChar(TEXT(']'), CloseBracket))
+			{
+				AddressPart = LocalEndPoint.Mid(1, CloseBracket - 1);
+				FString PortStr = LocalEndPoint.Mid(CloseBracket + 1);
+				if (PortStr.StartsWith(TEXT(":")))
+				{
+					LexFromString(Port, *PortStr.Mid(1));
+				}
+			}
+		}
+
+		if (!AddressPart.IsEmpty())
+		{
+			if (Port <= 0 || Port > 65535)
+			{
+				UE_LOG(LogBackChannel, Error, TEXT("Invalid or missing port in endpoint '%s'"), InEndPoint);
+				return false;
+			}
+
+			TargetAddr = SocketSubsystem->GetAddressFromString(AddressPart);
+			if (TargetAddr.IsValid())
+			{
+				TargetAddr->SetPort(Port);
+			}
+		}
+	}
+
+	if (!TargetAddr.IsValid())
+	{
+		UE_LOG(LogBackChannel, Error, TEXT("Failed to parse endpoint '%s'"), InEndPoint);
+		return false;
+	}
 
-	FSocket* NewSocket = ISocketSubsystem::Get(PLATFORM_SOCKETSUBSYSTEM)->CreateSocket(NAME_Stream, TEXT("FBackChannelConnection Client Socket"));
+	// Create socket with the correct protocol family for the target address
+	FName ProtocolType = TargetAddr->GetProtocolType();
+	FSocket* NewSocket = SocketSubsystem->CreateSocket(NAME_Stream, TEXT("FBackChannelConnection Client Socket"), ProtocolType);
 	if (NewSocket)
 	{
 		NewSocket->SetNonBlocking();
@@ -128,18 +182,11 @@ bool FBackChannelConnection::Connect(const TCHAR* InEndPoint)
 		// set buffer sizes
 		SetSocketBufferSizes(NewSocket, SendBufferSize, ReceiveBufferSize);
 
-		bool Success = false;
-
-		FIPv4Endpoint EndPointv4;
-		// Check for a valid IPv4 address string
-		if (FIPv4Endpoint::Parse(LocalEndPoint, EndPointv4))
-		{
-			Success = NewSocket->Connect(*EndPointv4.ToInternetAddr());
-		}
+		bool Success = NewSocket->Connect(*TargetAddr);
 
 		if (!Success)
 		{
-			ESocketErrors LastErr = ISocketSubsystem::Get(PLATFORM_SOCKETSUBSYSTEM)->GetLastErrorCode();
+			ESocketErrors LastErr = SocketSubsystem->GetLastErrorCode();
 
 			if (LastErr == SE_EINPROGRESS || LastErr == SE_EWOULDBLOCK)
 			{
@@ -147,7 +194,7 @@ bool FBackChannelConnection::Connect(const TCHAR* InEndPoint)
 			}
 			else
 			{
-				UE_LOG(LogBackChannel, Log, TEXT("Connect failed with error code (%d) error (%s)"), LastErr, ISocketSubsystem::Get(PLATFORM_SOCKETSUBSYSTEM)->GetSocketError(LastErr));
+				UE_LOG(LogBackChannel, Warning, TEXT("Connect failed with error code (%d) error (%s)"), LastErr, SocketSubsystem->GetSocketError(LastErr));
 			}
 		}
 
@@ -158,9 +205,8 @@ bool FBackChannelConnection::Connect(const TCHAR* InEndPoint)
 		}
 		else
 		{
-
 			CloseWithError(*FString::Printf(TEXT("Failed to open connection to %s."), InEndPoint), NewSocket);
-			ISocketSubsystem::Get(PLATFORM_SOCKETSUBSYSTEM)->DestroySocket(NewSocket);
+			SocketSubsystem->DestroySocket(NewSocket);
 		}
 	}
 
@@ -246,11 +292,38 @@ bool FBackChannelConnection::Listen(const int16 Port)
 	FScopeLock Lock(&SocketMutex);
 
 	ISocketSubsystem* SocketSubsystem = ISocketSubsystem::Get(PLATFORM_SOCKETSUBSYSTEM);
-	FSocket* NewSocket = SocketSubsystem->CreateSocket(NAME_Stream, TEXT("FBackChannelConnection Client Socket"));
-	if (NewSocket != nullptr)
+
+	// Try IPv6 dual-stack first (accepts both IPv4 and IPv6 connections)
+	FSocket* NewSocket = nullptr;
+	TSharedRef<FInternetAddr> BindAddr = SocketSubsystem->CreateInternetAddr();
+
+#if PLATFORM_HAS_BSD_IPV6_SOCKETS
+	NewSocket = SocketSubsystem->CreateSocket(NAME_Stream, TEXT("FBackChannelConnection Listener"), FNetworkProtocolTypes::IPv6);
+	if (NewSocket)
 	{
-		FIPv4Endpoint Endpoint(FIPv4Address::Any, Port);
+		// Platform socket subsystems already set IPV6_V6ONLY=false in CreateSocket(),
+		// enabling dual-stack: this IPv6 socket will accept IPv4 connections too.
+
+		// Bind to IPv6 any-address (::)
+		BindAddr = SocketSubsystem->CreateInternetAddr(FNetworkProtocolTypes::IPv6);
+		BindAddr->SetAnyAddress();
+		BindAddr->SetPort(Port);
+	}
+#endif
 
+	if (!NewSocket)
+	{
+		// Fall back to IPv4-only
+		NewSocket = SocketSubsystem->CreateSocket(NAME_Stream, TEXT("FBackChannelConnection Listener"), FNetworkProtocolTypes::IPv4);
+		if (NewSocket)
+		{
+			FIPv4Endpoint Endpoint(FIPv4Address::Any, Port);
+			BindAddr = Endpoint.ToInternetAddr();
+		}
+	}
+
+	if (NewSocket != nullptr)
+	{
 		bool Error = !NewSocket->SetReuseAddr(true);
 
 		if (!Error)
@@ -263,10 +336,9 @@ bool FBackChannelConnection::Listen(const int16 Port)
 			Error = !NewSocket->SetNonBlocking();
 		}
 
-
 		if (!Error)
 		{
-			Error = !NewSocket->Bind(*Endpoint.ToInternetAddr());
+			Error = !NewSocket->Bind(*BindAddr);
 		}
 
 		if (!Error)
diff --git a/Engine/Plugins/Experimental/NetworkServiceDiscovery/Source/NetworkServiceDiscovery/Private/Android/Java/NetworkServiceDiscoveryHelper.java b/Engine/Plugins/Experimental/NetworkServiceDiscovery/Source/NetworkServiceDiscovery/Private/Android/Java/NetworkServiceDiscoveryHelper.java
--- a/Engine/Plugins/Experimental/NetworkServiceDiscovery/Source/NetworkServiceDiscovery/Private/Android/Java/NetworkServiceDiscoveryHelper.java
+++ b/Engine/Plugins/Experimental/NetworkServiceDiscovery/Source/NetworkServiceDiscovery/Private/Android/Java/NetworkServiceDiscoveryHelper.java
@@ -5,9 +5,14 @@ package com.epicgames.unreal;
 import android.content.Context;
 import android.net.nsd.NsdManager;
 import android.net.nsd.NsdServiceInfo;
+import android.os.Build;
 import android.util.Log;
 
+import java.net.Inet4Address;
 import java.net.InetAddress;
+import java.net.Inet6Address;
+import java.net.NetworkInterface;
+import java.util.Enumeration;
 import java.util.HashMap;
 import java.util.Map;
 import java.util.concurrent.ConcurrentHashMap;
@@ -283,6 +288,121 @@ public class NetworkServiceDiscoveryHelper
 		return bIsDiscovering;
 	}
 
+	// ------------------------------------------------------------------
+	//  Route-check helpers
+	// ------------------------------------------------------------------
+
+	/**
+	 * Returns true if the device has at least one routable (non-link-local, non-APIPA)
+	 * IPv4 address on any active network interface.
+	 */
+	private static boolean hasRoutableIPv4Address()
+	{
+		try
+		{
+			Enumeration<NetworkInterface> interfaces = NetworkInterface.getNetworkInterfaces();
+			while (interfaces != null && interfaces.hasMoreElements())
+			{
+				NetworkInterface iface = interfaces.nextElement();
+				if (!iface.isUp() || iface.isLoopback())
+				{
+					continue;
+				}
+
+				Enumeration<InetAddress> addrs = iface.getInetAddresses();
+				while (addrs.hasMoreElements())
+				{
+					InetAddress addr = addrs.nextElement();
+					if (addr instanceof Inet4Address && !addr.isLoopbackAddress() && !addr.isLinkLocalAddress())
+					{
+						Log.d(TAG, "Device has routable IPv4 address: " + addr.getHostAddress() + " on " + iface.getName());
+						return true;
+					}
+				}
+			}
+		}
+		catch (Exception e)
+		{
+			Log.w(TAG, "Failed to enumerate network interfaces: " + e.getMessage());
+			return true; // Assume IPv4 is available if we can't check
+		}
+
+		Log.d(TAG, "Device has no routable IPv4 address");
+		return false;
+	}
+
+	/**
+	 * Finds the best routable address from a resolved service. On API 34+
+	 * (Android 14), NsdServiceInfo provides all addresses and we probe each
+	 * one. On older APIs only a single address is available so we use it as-is.
+	 */
+	private static InetAddress findRoutableAddress(NsdServiceInfo info)
+	{
+		if (Build.VERSION.SDK_INT >= 34)
+		{
+			return findRoutableAddressApi34(info);
+		}
+
+		// Pre-API 34: only one address available
+		InetAddress host = info.getHost();
+		if (host != null)
+		{
+			if (host instanceof Inet6Address && host.isLinkLocalAddress())
+			{
+				Log.d(TAG, "Skipping IPv6 " + host.getHostAddress() + " — link-local (pre-API 34)");
+				return null;
+			}
+			String addrType = (host instanceof Inet6Address) ? "IPv6" : "IPv4";
+			Log.d(TAG, "Candidate " + addrType + " " + host.getHostAddress() + " (single address, pre-API 34)");
+		}
+		return host;
+	}
+
+	/**
+	 * API 34+ path: iterates all resolved addresses and picks the best one.
+	 * Skips IPv4 candidates if the device has no routable IPv4 address.
+	 * Skips IPv6 link-local addresses. Returns the first viable candidate.
+	 */
+	private static InetAddress findRoutableAddressApi34(NsdServiceInfo info)
+	{
+		java.util.List<InetAddress> addresses = info.getHostAddresses();
+		boolean bHasRoutableIPv4 = hasRoutableIPv4Address();
+
+		// First pass: log all candidates
+		Log.d(TAG, "Service resolved with " + addresses.size() + " address(es) (device has routable IPv4: " + bHasRoutableIPv4 + "):");
+		for (InetAddress addr : addresses)
+		{
+			String addrType = (addr instanceof Inet6Address) ? "IPv6" : "IPv4";
+			boolean linkLocal = addr.isLinkLocalAddress();
+			Log.d(TAG, "  " + addrType + " " + addr.getHostAddress() + (linkLocal ? " (link-local)" : ""));
+		}
+
+		// Second pass: pick the first viable address
+		for (InetAddress candidate : addresses)
+		{
+			// Skip IPv6 link-local addresses
+			if (candidate instanceof Inet6Address && candidate.isLinkLocalAddress())
+			{
+				continue;
+			}
+
+			// Skip IPv4 candidates if the device has no routable IPv4
+			if (candidate instanceof Inet4Address && !bHasRoutableIPv4)
+			{
+				Log.d(TAG, "Skipping IPv4 " + candidate.getHostAddress() + " — device has no routable IPv4");
+				continue;
+			}
+
+			String addrType = (candidate instanceof Inet6Address) ? "IPv6" : "IPv4";
+			Log.d(TAG, "Selected " + addrType + " " + candidate.getHostAddress());
+			return candidate;
+		}
+
+		// No viable address found — fall back to legacy single-address API
+		Log.w(TAG, "No viable address found, falling back to " + info.getHost());
+		return info.getHost();
+	}
+
 	// ------------------------------------------------------------------
 	//  Service Resolution
 	// ------------------------------------------------------------------
@@ -313,8 +433,8 @@ public class NetworkServiceDiscoveryHelper
 			{
 				Log.d(TAG, "Resolved: " + info.getServiceName() + " -> " + info.getHost() + ":" + info.getPort());
 
-				InetAddress host = info.getHost();
-				String hostAddress = (host != null) ? host.getHostAddress() : "";
+				InetAddress routableHost = findRoutableAddress(info);
+				String hostAddress = (routableHost != null) ? routableHost.getHostAddress() : "";
 
 				Map<String, byte[]> attributes = info.getAttributes();
 				String[] keys = new String[attributes.size()];
diff --git a/Engine/Plugins/Experimental/NetworkServiceDiscovery/Source/NetworkServiceDiscovery/Private/Apple/NetworkServiceDiscoveryApple.mm b/Engine/Plugins/Experimental/NetworkServiceDiscovery/Source/NetworkServiceDiscovery/Private/Apple/NetworkServiceDiscoveryApple.mm
--- a/Engine/Plugins/Experimental/NetworkServiceDiscovery/Source/NetworkServiceDiscovery/Private/Apple/NetworkServiceDiscoveryApple.mm
+++ b/Engine/Plugins/Experimental/NetworkServiceDiscovery/Source/NetworkServiceDiscovery/Private/Apple/NetworkServiceDiscoveryApple.mm
@@ -13,11 +13,36 @@
 #include "Apple/PreAppleSystemHeaders.h"
 #import <Foundation/Foundation.h>
 #import <arpa/inet.h>
+#import <sys/socket.h>
+#import <unistd.h>
 #include "Apple/PostAppleSystemHeaders.h"
 
 #undef FVector
 #pragma pop_macro("FVector")
 
+// ============================================================================
+//  Route-check helper
+// ============================================================================
+
+/**
+ * Returns true if the kernel has a route to the given address.
+ * Uses the standard BSD UDP connect() probe: connecting a UDP socket
+ * triggers a route lookup without sending any traffic. If the kernel
+ * can find a route, connect() returns 0; otherwise it fails immediately.
+ * This is Apple's recommended technique for IPv6 transition support.
+ */
+static bool CanRouteToAddress(const struct sockaddr* Addr, socklen_t AddrLen)
+{
+	int Sock = socket(Addr->sa_family, SOCK_DGRAM, 0);
+	if (Sock < 0)
+	{
+		return false;
+	}
+	int Result = connect(Sock, Addr, AddrLen);
+	close(Sock);
+	return (Result == 0);
+}
+
 // ============================================================================
 //  Objective-C delegate for NSNetServiceBrowser (discovery)
 // ============================================================================
@@ -177,18 +202,80 @@ - (void)netServiceDidResolveAddress:(NSNetService*)Service
 	Info.Port = (int32)Service.port;
 	Info.bIsResolved = true;
 
-	// Extract IPv4 address from the resolved addresses
+	// Pick the best address from the resolved set by probing reachability.
+	// NSNetService returns both IPv4 and IPv6 addresses but doesn't tell us
+	// which address family the mDNS reply actually came in on. On an IPv6-only
+	// network the device may have a non-routable APIPA (169.254.x.x) IPv4
+	// address, so blindly preferring IPv4 would pick an unreachable host.
+	// Instead we iterate the addresses in the order the system provides them
+	// and use the first one the kernel can actually route to (UDP connect probe).
+	NSString* FallbackAddrStr = nil;
 	for (NSData* AddressData in Service.addresses)
 	{
 		const struct sockaddr* SockAddr = (const struct sockaddr*)[AddressData bytes];
+		socklen_t AddrLen = (socklen_t)[AddressData length];
+
 		if (SockAddr->sa_family == AF_INET)
 		{
 			const struct sockaddr_in* Addr4 = (const struct sockaddr_in*)SockAddr;
-			char AddrStr[INET_ADDRSTRLEN];
-			inet_ntop(AF_INET, &Addr4->sin_addr, AddrStr, sizeof(AddrStr));
-			Info.Address = FString(ANSI_TO_TCHAR(AddrStr));
-			break;
+			char AddrBuf[INET_ADDRSTRLEN];
+			inet_ntop(AF_INET, &Addr4->sin_addr, AddrBuf, sizeof(AddrBuf));
+			NSString* Str = [NSString stringWithUTF8String:AddrBuf];
+
+			bool bRoutable = CanRouteToAddress(SockAddr, AddrLen);
+			UE_LOG(LogNetworkServiceDiscovery, Log, TEXT("Apple: Candidate IPv4 %s — %s"),
+				*FString(Str), bRoutable ? TEXT("routable") : TEXT("not routable"));
+
+			if (FallbackAddrStr == nil)
+			{
+				FallbackAddrStr = Str;
+			}
+			if (bRoutable)
+			{
+				Info.Address = FString(Str);
+				break;
+			}
 		}
+		else if (SockAddr->sa_family == AF_INET6)
+		{
+			const struct sockaddr_in6* Addr6 = (const struct sockaddr_in6*)SockAddr;
+
+			// Skip link-local addresses (fe80::) — they require scope IDs
+			// and are unreliable for cross-device connections
+			if (IN6_IS_ADDR_LINKLOCAL(&Addr6->sin6_addr))
+			{
+				char AddrBuf[INET6_ADDRSTRLEN];
+				inet_ntop(AF_INET6, &Addr6->sin6_addr, AddrBuf, sizeof(AddrBuf));
+				UE_LOG(LogNetworkServiceDiscovery, Log, TEXT("Apple: Candidate IPv6 %s — skipped (link-local)"),
+					*FString([NSString stringWithUTF8String:AddrBuf]));
+				continue;
+			}
+
+			char AddrBuf[INET6_ADDRSTRLEN];
+			inet_ntop(AF_INET6, &Addr6->sin6_addr, AddrBuf, sizeof(AddrBuf));
+			NSString* Str = [NSString stringWithUTF8String:AddrBuf];
+
+			bool bRoutable = CanRouteToAddress(SockAddr, AddrLen);
+			UE_LOG(LogNetworkServiceDiscovery, Log, TEXT("Apple: Candidate IPv6 %s — %s"),
+				*FString(Str), bRoutable ? TEXT("routable") : TEXT("not routable"));
+
+			if (FallbackAddrStr == nil)
+			{
+				FallbackAddrStr = Str;
+			}
+			if (bRoutable)
+			{
+				Info.Address = FString(Str);
+				break;
+			}
+		}
+	}
+
+	// If no address was routable (shouldn't happen), use the first candidate
+	if (Info.Address.IsEmpty() && FallbackAddrStr != nil)
+	{
+		UE_LOG(LogNetworkServiceDiscovery, Warning, TEXT("Apple: No routable address found, falling back to %s"), *FString(FallbackAddrStr));
+		Info.Address = FString(FallbackAddrStr);
 	}
 
 	// Extract TXT record
diff --git a/Engine/Plugins/Experimental/NetworkServiceDiscovery/Source/NetworkServiceDiscovery/Private/Windows/NetworkServiceDiscoveryWindows.cpp b/Engine/Plugins/Experimental/NetworkServiceDiscovery/Source/NetworkServiceDiscovery/Private/Windows/NetworkServiceDiscoveryWindows.cpp
--- a/Engine/Plugins/Experimental/NetworkServiceDiscovery/Source/NetworkServiceDiscovery/Private/Windows/NetworkServiceDiscoveryWindows.cpp
+++ b/Engine/Plugins/Experimental/NetworkServiceDiscovery/Source/NetworkServiceDiscovery/Private/Windows/NetworkServiceDiscoveryWindows.cpp
@@ -364,6 +364,12 @@ static void WINAPI NSD_ResolveCallback(DWORD Status, PVOID Context, PDNS_SERVICE
 		InetNtopW(AF_INET, &Addr, AddrStr, INET_ADDRSTRLEN);
 		Info.Address = FString(AddrStr);
 	}
+	else if (pInstance->ip6Address)
+	{
+		WCHAR AddrStr[INET6_ADDRSTRLEN];
+		InetNtopW(AF_INET6, pInstance->ip6Address, AddrStr, INET6_ADDRSTRLEN);
+		Info.Address = FString(AddrStr);
+	}
 
 	for (DWORD i = 0; i < pInstance->dwPropertyCount; i++)
 	{
diff --git a/Engine/Plugins/Experimental/NetworkServiceDiscovery/Source/NetworkServiceDiscovery/Public/NetworkServiceDiscoveryTypes.h b/Engine/Plugins/Experimental/NetworkServiceDiscovery/Source/NetworkServiceDiscovery/Public/NetworkServiceDiscoveryTypes.h
--- a/Engine/Plugins/Experimental/NetworkServiceDiscovery/Source/NetworkServiceDiscovery/Public/NetworkServiceDiscoveryTypes.h
+++ b/Engine/Plugins/Experimental/NetworkServiceDiscovery/Source/NetworkServiceDiscovery/Public/NetworkServiceDiscoveryTypes.h
@@ -18,7 +18,7 @@ struct FNetworkServiceInfo
 	/** Resolved hostname of the service provider */
 	FString HostName;
 
-	/** Resolved IP address (IPv4 string) */
+	/** Resolved IP address (IPv4 or IPv6 string, e.g. "192.168.1.100" or "2001:db8::1") */
 	FString Address;
 
 	/** Resolved port number */
diff --git a/Engine/Plugins/Experimental/RemoteSession/Source/RemoteSession/Private/RemoteSessionClient.cpp b/Engine/Plugins/Experimental/RemoteSession/Source/RemoteSession/Private/RemoteSessionClient.cpp
--- a/Engine/Plugins/Experimental/RemoteSession/Source/RemoteSession/Private/RemoteSessionClient.cpp
+++ b/Engine/Plugins/Experimental/RemoteSession/Source/RemoteSession/Private/RemoteSessionClient.cpp
@@ -19,9 +19,43 @@ FRemoteSessionClient::FRemoteSessionClient(const TCHAR* InHostAddress)
 
 	IsConnecting = false;
 
-	if (HostAddress.Contains(TEXT(":")) == false)
+	// Detect whether the address already includes a port
+	bool bHasPort = false;
+	if (HostAddress.StartsWith(TEXT("[")))
 	{
-		HostAddress += FString::Printf(TEXT(":%d"), (int32)IRemoteSessionModule::kDefaultPort);
+		// IPv6 bracket notation — port follows the closing bracket, e.g. [::1]:2049
+		int32 CloseBracket;
+		if (HostAddress.FindChar(TEXT(']'), CloseBracket))
+		{
+			bHasPort = HostAddress.Mid(CloseBracket + 1).StartsWith(TEXT(":"));
+		}
+	}
+	else
+	{
+		// A single colon means IPv4/hostname with port (e.g. "192.168.1.1:2049").
+		// Multiple colons means bare IPv6 without port (e.g. "::1", "fd00::1").
+		int32 FirstColon = HostAddress.Find(TEXT(":"), ESearchCase::IgnoreCase, ESearchDir::FromStart);
+		int32 LastColon = HostAddress.Find(TEXT(":"), ESearchCase::IgnoreCase, ESearchDir::FromEnd);
+		bHasPort = (FirstColon != INDEX_NONE && FirstColon == LastColon);
+	}
+
+	if (!bHasPort)
+	{
+		if (HostAddress.StartsWith(TEXT("[")))
+		{
+			// [::1] -> [::1]:2049
+			HostAddress += FString::Printf(TEXT(":%d"), (int32)IRemoteSessionModule::kDefaultPort);
+		}
+		else if (HostAddress.Contains(TEXT(":")))
+		{
+			// Bare IPv6 like "::1" -> "[::1]:2049"
+			HostAddress = FString::Printf(TEXT("[%s]:%d"), *HostAddress, (int32)IRemoteSessionModule::kDefaultPort);
+		}
+		else
+		{
+			// IPv4 or hostname -> "host:2049"
+			HostAddress += FString::Printf(TEXT(":%d"), (int32)IRemoteSessionModule::kDefaultPort);
+		}
 	}
 
 	// OnReceiveHello will enable PixelStreaming if supported by both host and client
diff --git a/Engine/Source/Runtime/Core/Public/Android/AndroidPlatform.h b/Engine/Source/Runtime/Core/Public/Android/AndroidPlatform.h
--- a/Engine/Source/Runtime/Core/Public/Android/AndroidPlatform.h
+++ b/Engine/Source/Runtime/Core/Public/Android/AndroidPlatform.h
@@ -48,6 +48,7 @@ typedef FAndroidTypes FPlatformTypes;
 #define PLATFORM_HAS_NO_EPROCLIM					1
 #define PLATFORM_USES_GLES							1
 #define PLATFORM_BUILTIN_VERTEX_HALF_FLOAT			0
+#define PLATFORM_HAS_BSD_IPV6_SOCKETS				1
 #define PLATFORM_HAS_BSD_SOCKET_FEATURE_IOCTL		1
 #define PLATFORM_HAS_BSD_SOCKET_FEATURE_MSG_DONTWAIT	1
 #define PLATFORM_HAS_TOUCH_MAIN_SCREEN				1
diff --git a/Engine/Source/Runtime/Sockets/Private/Android/SocketSubsystemAndroid.cpp b/Engine/Source/Runtime/Sockets/Private/Android/SocketSubsystemAndroid.cpp
--- a/Engine/Source/Runtime/Sockets/Private/Android/SocketSubsystemAndroid.cpp
+++ b/Engine/Source/Runtime/Sockets/Private/Android/SocketSubsystemAndroid.cpp
@@ -111,6 +111,25 @@ const TCHAR* FSocketSubsystemAndroid::GetSocketAPIName() const
 	return TEXT("BSD_Android");
 }
 
+/**
+ * Creates a socket and ensures IPv6 sockets are dual-stack capable (IPV6_V6ONLY=false).
+ * This matches the behavior of iOS, Unix, and Windows socket subsystems, allowing a
+ * single IPv6 socket to accept both IPv4 and IPv6 connections.
+ */
+FSocket* FSocketSubsystemAndroid::CreateSocket(const FName& SocketType, const FString& SocketDescription, const FName& ProtocolType)
+{
+	FSocketBSD* NewSocket = (FSocketBSD*)FSocketSubsystemBSD::CreateSocket(SocketType, SocketDescription, ProtocolType);
+	if (NewSocket != nullptr)
+	{
+		NewSocket->SetIPv6Only(false);
+	}
+	else
+	{
+		UE_LOG(LogSockets, Warning, TEXT("Failed to create socket %s [%s]"), *SocketType.ToString(), *SocketDescription);
+	}
+	return NewSocket;
+}
+
 /** Priority values for Android adapters. This code is copied from the iOS system as well. */
 enum class EAdapterPriorityValues : uint8
 {
diff --git a/Engine/Source/Runtime/Sockets/Private/Android/SocketSubsystemAndroid.h b/Engine/Source/Runtime/Sockets/Private/Android/SocketSubsystemAndroid.h
--- a/Engine/Source/Runtime/Sockets/Private/Android/SocketSubsystemAndroid.h
+++ b/Engine/Source/Runtime/Sockets/Private/Android/SocketSubsystemAndroid.h
@@ -79,6 +79,8 @@ public:
 
 	virtual int32 GetAddressInfoHintFlag(EAddressInfoFlags InFlags) const override;
 
+	virtual FSocket* CreateSocket(const FName& SocketType, const FString& SocketDescription, const FName& ProtocolType) override;
+
 	virtual bool GetLocalAdapterAddresses(TArray<TSharedPtr<FInternetAddr>>& OutAddresses) override;
 
 	virtual FName GetDefaultSocketProtocolFamily() const override
